Media: Initial revision of media framework

diff --git a/Engine/Source/Runtime/Launch/Launch.Build.cs b/Engine/Source/Runtime/Launch/Launch.Build.cs
--- a/Engine/Source/Runtime/Launch/Launch.Build.cs
+++ b/Engine/Source/Runtime/Launch/Launch.Build.cs
@@ -23,6 +23,7 @@ public class Launch : ModuleRules
 				"CoreUObject",
 				"Engine",
 				"InputCore",
+                "Media",
                 "MoviePlayer",
 				"Networking",
 				"PakFile",
diff --git a/Engine/Source/Runtime/Media/Media.Build.cs b/Engine/Source/Runtime/Media/Media.Build.cs
new file mode 100644
--- /dev/null
+++ b/Engine/Source/Runtime/Media/Media.Build.cs
@@ -0,0 +1,23 @@
+// Copyright 1998-2014 Epic Games, Inc. All Rights Reserved.
+
+namespace UnrealBuildTool.Rules
+{
+	public class Media : ModuleRules
+	{
+		public Media(TargetInfo Target)
+		{
+			PublicDependencyModuleNames.AddRange(
+				new string[]
+				{
+					"Core",
+				}
+			); 
+
+			PrivateIncludePaths.AddRange(
+				new string[] {
+					"Runtime/Media/Private",
+				}
+			);
+		}
+	}
+}
diff --git a/Engine/Source/Runtime/Media/Private/MediaModule.cpp b/Engine/Source/Runtime/Media/Private/MediaModule.cpp
new file mode 100644
--- /dev/null
+++ b/Engine/Source/Runtime/Media/Private/MediaModule.cpp
@@ -0,0 +1,102 @@
+// Copyright 1998-2014 Epic Games, Inc. All Rights Reserved.
+
+#include "MediaPrivatePCH.h"
+#include "ModuleManager.h"
+
+
+/**
+ * Implements the Media module.
+ */
+class FMediaModule
+	: public IMediaModule
+{
+public:
+
+	// IModuleInterface interface
+
+	virtual void StartupModule( ) override { }
+
+	virtual void ShutdownModule( ) override
+	{
+		MediaPlayerFactories.Reset();
+	}
+
+	virtual bool SupportsDynamicReloading( ) override
+	{
+		return false;
+	}
+
+public:
+
+	// IMediaModuleInterface
+
+	virtual TSharedPtr<IMediaPlayer> CreatePlayer( const FString& Url ) override
+	{
+		TSharedPtr<IMediaPlayer> Player;
+
+		for (IMediaPlayerFactory* PlayerFactory : MediaPlayerFactories)
+		{
+			if (PlayerFactory->SupportsFile(Url))
+			{
+				Player = PlayerFactory->CreatePlayer();
+
+				if (Player.IsValid())
+				{
+					break;
+				}
+			}
+		}
+
+		return Player;
+	}
+
+	virtual int32 GetSupportedFormats( FMediaFormats& OutFormats ) override
+	{
+		OutFormats.Reset();
+
+		for (IMediaPlayerFactory* Factory : MediaPlayerFactories)
+		{
+			OutFormats.Append(Factory->GetSupportedFormats());
+		}
+	
+		return OutFormats.Num();
+	}
+
+	DECLARE_DERIVED_EVENT(FMediaModule, IMediaModule::FOnFactoryAdded, FOnFactoryAdded);
+	virtual FOnFactoryAdded& OnFactoryAdded( ) override
+	{
+		return FactoryAddedEvent;
+	}
+
+	DECLARE_DERIVED_EVENT(FMediaModule, IMediaModule::FOnFactoryRemoved, FOnFactoryRemoved);
+	virtual FOnFactoryRemoved& OnFactoryRemoved( ) override
+	{
+		return FactoryRemovedEvent;
+	}
+
+	virtual void RegisterPlayerFactory( IMediaPlayerFactory& Factory ) override
+	{
+		MediaPlayerFactories.AddUnique(&Factory);
+		FactoryAddedEvent.Broadcast();
+	}
+
+	virtual void UnregisterPlayerFactory( IMediaPlayerFactory& Factory ) override
+	{
+		MediaPlayerFactories.Remove(&Factory);
+		FactoryRemovedEvent.Broadcast();
+	}
+
+private:
+
+	// Holds an event delegate that is invoked after a video player factory has been added.
+	FOnFactoryAdded FactoryAddedEvent;
+
+	// Holds an event delegate that is invoked after a video player factory has been removed.
+	FOnFactoryRemoved FactoryRemovedEvent;
+
+	// Holds the registered video player factories.
+	TArray<IMediaPlayerFactory*> MediaPlayerFactories;
+};
+
+
+IMPLEMENT_MODULE(FMediaModule, Media);
diff --git a/Engine/Source/Runtime/Media/Private/MediaPrivatePCH.h b/Engine/Source/Runtime/Media/Private/MediaPrivatePCH.h
new file mode 100644
--- /dev/null
+++ b/Engine/Source/Runtime/Media/Private/MediaPrivatePCH.h
@@ -0,0 +1,11 @@
+// Copyright 1998-2014 Epic Games, Inc. All Rights Reserved.
+
+#pragma once
+
+#include "Media.h"
+
+
+/* Private dependencies
+ *****************************************************************************/
+
+#include "Core.h"
diff --git a/Engine/Source/Runtime/Media/Public/Common/MediaSampleBuffer.h b/Engine/Source/Runtime/Media/Public/Common/MediaSampleBuffer.h
new file mode 100644
--- /dev/null
+++ b/Engine/Source/Runtime/Media/Public/Common/MediaSampleBuffer.h
@@ -0,0 +1,72 @@
+// Copyright 1998-2014 Epic Games, Inc. All Rights Reserved.
+
+#pragma once
+
+
+/** Type definition for shared pointers to instances of FMediaSampleBuffer. */
+typedef TSharedPtr<class FMediaSampleBuffer, ESPMode::ThreadSafe> FMediaSampleBufferPtr;
+
+/** Type definition for shared references to instances of FMediaSampleBuffer. */
+typedef TSharedRef<class FMediaSampleBuffer, ESPMode::ThreadSafe> FMediaSampleBufferRef;
+
+
+/**
+ * Implements a media sample sink that buffers the last sample.
+ */
+class FMediaSampleBuffer
+	: public TQueue<TSharedPtr<TArray<uint8>, ESPMode::ThreadSafe>>
+	, public IMediaSink
+{
+public:
+
+	/** Default constructor. */
+	FMediaSampleBuffer( )
+		: CurrentSampleTime(FTimespan::MinValue())
+	{ }
+
+public:
+
+	/**
+	 * Gets the sample data of the current media sample.
+	 *
+	 * @return The current sample.
+	 * @see GetCurrentSampleTime
+	 */
+	TSharedPtr<TArray<uint8>, ESPMode::ThreadSafe> GetCurrentSample( ) const
+	{
+		return CurrentSample;
+	}
+	
+	/**
+	 * Gets the playback time of the currently available media sample.
+	 *
+	 * @return Sample playback time.
+	 * @see GetCurrentSample
+	 */
+	FTimespan GetCurrentSampleTime( ) const
+	{
+		return CurrentSampleTime;
+	}
+
+public:
+
+	// IMediaSink interface
+
+	void ProcessMediaSample( const void* Buffer, uint32 BufferSize, FTimespan Duration, FTimespan Time ) override
+	{
+		TArray<uint8>* Sample = new TArray<uint8>();
+		Sample->AddUninitialized(BufferSize);
+		FMemory::Memcpy(Sample->GetData(), Buffer, BufferSize);
+
+		CurrentSample = MakeShareable(Sample);
+		CurrentSampleTime = Time;
+	}
+	
+private:
+
+	// Holds the sample data of the current media sample.
+	TSharedPtr<TArray<uint8>, ESPMode::ThreadSafe> CurrentSample;
+
+	// The playback time of the currently available media sample.
+	FTimespan CurrentSampleTime;
+};
diff --git a/Engine/Source/Runtime/Media/Public/Common/MediaSampleQueue.h b/Engine/Source/Runtime/Media/Public/Common/MediaSampleQueue.h
new file mode 100644
--- /dev/null
+++ b/Engine/Source/Runtime/Media/Public/Common/MediaSampleQueue.h
@@ -0,0 +1,32 @@
+// Copyright 1998-2014 Epic Games, Inc. All Rights Reserved.
+
+#pragma once
+
+
+/** Type definition for shared pointers to instances of FMediaSampleQueue. */
+typedef TSharedPtr<class FMediaSampleQueue, ESPMode::ThreadSafe> FMediaSampleQueuePtr;
+
+/** Type definition for shared references to instances of FMediaSampleQueue. */
+typedef TSharedRef<class FMediaSampleQueue, ESPMode::ThreadSafe> FMediaSampleQueueRef;
+
+
+/**
+ * Implements a media sample sink that queues up samples.
+ */
+class FMediaSampleQueue
+	: public TQueue<TSharedPtr<TArray<uint8>, ESPMode::ThreadSafe>>
+	, public IMediaSink
+{
+public:
+
+	// IMediaSink interface
+
+	void ProcessMediaSample( const void* Buffer, uint32 BufferSize, FTimespan Duration, FTimespan Time ) override
+	{
+		TArray<uint8>* Sample = new TArray<uint8>();
+		Sample->AddUninitialized(BufferSize);
+		FMemory::Memcpy(Sample->GetData(), Buffer, BufferSize);
+
+		Enqueue(MakeShareable(Sample));
+	}
+};
diff --git a/Engine/Source/Runtime/Media/Public/Interfaces/IMediaModule.h b/Engine/Source/Runtime/Media/Public/Interfaces/IMediaModule.h
new file mode 100644
--- /dev/null
+++ b/Engine/Source/Runtime/Media/Public/Interfaces/IMediaModule.h
@@ -0,0 +1,89 @@
+// Copyright 1998-2014 Epic Games, Inc. All Rights Reserved.
+
+#pragma once
+
+
+/**
+ * Interface for media modules.
+ */
+class IMediaModule
+	: public IModuleInterface
+{
+public:
+
+	/**
+	 * Creates a media player for the specified media URL.
+	 *
+	 * @param Url The URL to the media file container to create the player for.
+	 */
+	virtual TSharedPtr<IMediaPlayer> CreatePlayer( const FString& Url ) = 0;
+
+	/**
+	 * Gets the collection of supported media formats.
+	 *
+	 * @param OutFormats Will hold the supported formats.
+	 */
+	virtual int32 GetSupportedFormats( FMediaFormats& OutFormats ) = 0;
+
+	/**
+	 * Registers a media player factory.
+	 *
+	 * @param Factory The media player factory to register.
+	 * @see UnregisterPlayerFactory
+	 */
+	virtual void RegisterPlayerFactory( IMediaPlayerFactory& Factory ) = 0;
+
+	/**
+	 * Unregisters a media player factory.
+	 *
+	 * @param Factory The media player factory to unregister.
+	 * @see RegisterPlayerFactory
+	 */
+	virtual void UnregisterPlayerFactory( IMediaPlayerFactory& Factory ) = 0;
+
+public:
+
+	/**
+	 * Gets an event delegate that is invoked after a media player factory has been added.
+	 *
+	 * @return The event delegate.
+	 * @see FOnFactoryRemoved
+	 */
+	DECLARE_EVENT(IMediaModule, FOnFactoryAdded);
+	virtual FOnFactoryAdded& OnFactoryAdded( ) = 0;
+
+	/**
+	 * Gets an event delegate that is invoked after a media player factory has been removed.
+	 *
+	 * @return The event delegate.
+	 * @see FOnFactoryAdded
+	 */
+	DECLARE_EVENT(IMediaModule, FOnFactoryRemoved);
+	virtual FOnFactoryRemoved& OnFactoryRemoved( ) = 0;
+
+public:
+
+	/**
+	 * Gets a reference to the messaging module instance.
+	 *
+	 * @todo gmp: better implementation using dependency injection.
+	 *
+	 * @return A reference to the Media module.
+	 */
+	static IMediaModule& Get( )
+	{
+#if PLATFORM_IOS
+        static IMediaModule& Media = FModuleManager::LoadModuleChecked<IMediaModule>("Media");
+        return MediaModule;
+#else
+        return FModuleManager::LoadModuleChecked<IMediaModule>("Media");
+#endif
+	}
+
+public:
+
+	/**
+	 * Virtual destructor.
+	 */
+	virtual ~IMediaModule( ) { }
+};
diff --git a/Engine/Source/Runtime/Media/Public/Interfaces/IMediaPlayer.h b/Engine/Source/Runtime/Media/Public/Interfaces/IMediaPlayer.h
new file mode 100644
--- /dev/null
+++ b/Engine/Source/Runtime/Media/Public/Interfaces/IMediaPlayer.h
@@ -0,0 +1,280 @@
+// Copyright 1998-2014 Epic Games, Inc. All Rights Reserved.
+
+#pragma once
+
+
+/**
+ * Enumerates directions for seeking in media.
+ */
+enum class EMediaSeekDirection
+{
+	/** Seek backwards from current position. */
+	Backward,
+
+	/** Seek from the beginning of the media. */
+	Beginning,
+
+	/** Seek from the end of the media. */
+	End,
+
+	/** Seek forward from current position. */
+	Forward
+};
+
+
+/** Type definition for shared pointers to instances of IMediaPlayer. */
+typedef TSharedPtr<class IMediaPlayer> IMediaPlayerPtr;
+
+/** Type definition for shared references to instances of IMediaPlayer. */
+typedef TSharedRef<class IMediaPlayer> IMediaPlayerRef;
+
+
+/**
+ * Interface for media players.
+ *
+ * @see IMediaTrack
+ */
+class IMediaPlayer
+{
+public:
+
+	/**
+	 * Closes a previously opened media.
+	 *
+	 * If no media has been opened, this function has no effect.
+	 *
+	 * @see IsReady, Open
+	 */
+	virtual void Close( ) = 0;
+
+	/**
+	 * Gets the media's duration.
+	 *
+	 * @return A time span representing the duration.
+	 * @see GetTime, Seek
+	 */
+	virtual FTimespan GetDuration( ) const = 0;
+
+	/**
+	 * Gets the last error that occurred during media loading or playback.
+	 *
+	 * @return The error string, or an empty string if no error occurred.
+	 */
+//	virtual FString GetLastError( ) const = 0;
+
+	/**
+	 * Gets the nominal playback rate, i.e. 1.0 for real time.
+	 *
+	 * @return Playback rate.
+	 * @see Play
+	 */
+	virtual float GetRate( ) const = 0;
+
+	/**
+	 * Gets the player's current playback time.
+	 *
+	 * @return Playback time.
+	 * @see GetDuration, Seek
+	 */
+	virtual FTimespan GetTime( ) const = 0;
+
+	/**
+	 * Gets access to the media's audio, video and other tracks.
+	 *
+	 * @return Media tracks interface.
+	 */
+	virtual const TArray<IMediaTrackRef>& GetTracks( ) const = 0;
+
+	/**
+	 * Gets the URL of the currently loaded media.
+	 *
+	 * @return Media URL.
+	 */
+	virtual FString GetUrl( ) const = 0;
+
+	/**
+	 * Checks whether media playback is currently paused.
+	 *
+	 * @return true if playback is paused, false otherwise.
+	 * @see IsPlaying, IsStopped, Pause, Play, Stop
+	 */
+	virtual bool IsPaused( ) const = 0;
+
+	/**
+	 * Checks whether media playback has currently in progress.
+	 *
+	 * @return true if playback is in progress, false otherwise.
+	 * @see IsPaused, IsStopped, Pause, Play, Stop
+	 */
+	virtual bool IsPlaying( ) const = 0;
+
+	/**
+	 * Checks whether this player is ready for playback.
+	 *
+	 * A media player is considered ready if some media has been opened
+	 * successfully using the Open method and no error occurred during
+	 * loading or playback.
+	 *
+	 * @return true if ready, false otherwise.
+	 * @see Close, Open
+	 */
+	virtual bool IsReady( ) const = 0;
+
+	/**
+	 * Opens a media from a file on disk.
+	 *
+	 * @param Url The URL of the media to open (file name or web address).
+	 * @return true if the media was opened successfully, false otherwise.
+	 * @see Close, IsReady
+	 */
+	virtual bool Open( const FString& Url ) = 0;
+
+	/**
+	 * Opens a media from a buffer.
+	 *
+	 * @param Buffer The buffer holding the media data.
+	 * @param OriginalUrl The original URL of the media that was loaded into the buffer.
+	 * @return true if the media was opened, false otherwise.
+	 * @see Close, IsReady
+	 */
+	virtual bool Open( const TSharedRef<TArray<uint8>>& Buffer, const FString& OriginalUrl ) = 0;
+
+	/**
+	 * Starts media playback at the specified rate.
+	 *
+	 * A playback rate of 1.0 will play the media normally at real-time.
+	 * A rate of 0.0 corresponds to pausing playback. A negative rate, if
+	 * supported, plays the media in reverse, and a rate larger than 1.0
+	 * fast forwards playback.
+	 *
+	 * @param Rate The playback rate to use.
+	 * @return true if the media will be played, false otherwise.
+	 * @see GetRate, Pause, Stop
+	 */
+	virtual bool Play( float Rate ) = 0;
+
+	/**
+	 * Changes the media's playback time.
+	 *
+	 * @param Time The playback time to set.
+	 * @return true on success, false otherwise.
+	 * @see GetDuration, GetTime
+	 */
+	virtual bool Seek( const FTimespan& Time ) = 0;
+
+	/**
+	 * Checks whether the specified playback rate is supported.
+	 *
+	 * @param Rate The rate to check (can be negative for reverse play).
+	 * @param Unthinned Whether no frames should be dropped at the given rate.
+	 * @return true if the rate is supported, false otherwise.
+	 * @see SupportsScrubbing, SupportsSeeking
+	 */
+	virtual bool SupportsRate( float Rate, bool Unthinned ) const = 0;
+
+	/**
+	 * Checks whether scrubbing is supported.
+	 *
+	 * Scrubbing is the ability to decode video frames while seeking in a media item at a playback rate of 0.0.
+	 *
+	 * @return true if scrubbing is supported, false otherwise.
+	 * @see SupportsRate, SupportsSeeking
+	 */
+	virtual bool SupportsScrubbing( ) const = 0;
+
+	/**
+	 * Checks whether the currently loaded media can jump to certain times.
+	 *
+	 * @return true if seeking is supported, false otherwise.
+	 * @see SupportsRate, SupportsScrubbing
+	 */
+	virtual bool SupportsSeeking( ) const = 0;
+
+public:
+
+	/** Gets an event delegate that is invoked when media has been closed. */
+	DECLARE_EVENT_OneParam(IMediaPlayer, FOnMediaClosing, FString /*ClosedUrl*/)
+	virtual FOnMediaClosing& OnClosing( ) = 0;
+
+	/** Gets an event delegate that is invoked when media has been opened. */
+	DECLARE_EVENT_OneParam(IMediaPlayer, FOnMediaOpened, FString /*OpenedUrl*/)
+	virtual FOnMediaOpened& OnOpened( ) = 0;
+
+public:
+
+	/**
+	 * Attempts to get the media track with the specified index and type.
+	 *
+	 * @param TrackIndex The index of the track to get.
+	 * @param TrackType The expected type of the track, i.e. audio or video.
+	 * @return The track if it exists, nullptr otherwise.
+	 */
+	IMediaTrackPtr GetTrackSafe( int32 TrackIndex, EMediaTrackTypes TrackType )
+	{
+		if (GetTracks().IsValidIndex(TrackIndex))
+		{
+			IMediaTrackPtr Track = GetTracks()[TrackIndex];
+
+			if (Track->GetType() == TrackType)
+			{
+				return Track;
+			}
+		}
+
+		return nullptr;
+	}
+
+	/**
+	 * Pauses media playback.
+	 *
+	 * @return true if the media is being paused, false otherwise.
+	 * @see Play, Stop
+	 */
+	FORCEINLINE bool Pause( )
+	{
+		return Play(0.0f);
+	}
+
+	/**
+	 * Starts media playback at the default rate of 1.0.
+	 *
+	 * This method is equivalent to setting the playback rate to 1.0.
+	 *
+	 * @return true if playback is starting, false otherwise.
+	 * @see Pause, Stop
+	 */
+	FORCEINLINE bool Play( )
+	{
+		return Play(1.0f);
+	}
+
+	/**
+	 * Changes the playback time of the media by a relative offset in the given direction.
+	 *
+	 * @param TimeOffset The offset to apply to the time.
+	 * @param Direction The direction to seek in.
+	 * @return true on success, false otherwise.
+	 * @see GetDuration, GetTime
+	 */
+	bool Seek( const FTimespan& TimeOffset, EMediaSeekDirection Direction )
+	{
+		FTimespan SeekTime;
+
+		switch (Direction)
+		{
+		case EMediaSeekDirection::Backward: SeekTime = GetTime() - TimeOffset; break;
+		case EMediaSeekDirection::Beginning: SeekTime = TimeOffset; break;
+		case EMediaSeekDirection::End: SeekTime = GetDuration() - TimeOffset; break;
+		case EMediaSeekDirection::Forward: SeekTime = GetTime() + TimeOffset; break;
+		}
+
+		return Seek(SeekTime);
+	}
+
+public:
+
+	/**
+	 * Virtual destructor.
+	 */
+	virtual ~IMediaPlayer( ) { }
+};
diff --git a/Engine/Source/Runtime/Media/Public/Interfaces/IMediaPlayerFactory.h b/Engine/Source/Runtime/Media/Public/Interfaces/IMediaPlayerFactory.h
new file mode 100644
--- /dev/null
+++ b/Engine/Source/Runtime/Media/Public/Interfaces/IMediaPlayerFactory.h
@@ -0,0 +1,62 @@
+// Copyright 1998-2014 Epic Games, Inc. All Rights Reserved.
+
+#pragma once
+
+
+/**
+ * Type definition for media file container formats.
+ *
+ * This type is a TMap that maps the file extensions of supported media file containers
+ * formats to a human readable description text. The file extensions are stored without
+ * a leading dot, i.e. {"avi", "Audio Video Interleave File"}.
+ */
+typedef TMap<FString, FText> FMediaFormats;
+
+
+/**
+ * Interface for media player factories.
+ *
+ * Video player factories are used to create instances of media player implementations.
+ * Most media players will be implemented inside plug-ins, which will register their
+ * factories on startup. By associating a media format with a factory, the Video module
+ * can provide available media players for a given format.
+ */
+class IMediaPlayerFactory
+{
+public:
+
+	/**
+	 * Creates a media player.
+	 *
+	 * @return A new media player, or nullptr if a player couldn't be created.
+	 */
+	virtual TSharedPtr<IMediaPlayer> CreatePlayer( ) = 0;
+
+	/**
+	 * Gets the collection of supported media formats.
+	 *
+	 * @return A collection of supported formats.
+	 * @see SupportsFile
+	 */
+	virtual const FMediaFormats& GetSupportedFormats( ) const = 0;
+
+public:
+
+	/**
+	 * Checks whether this factory can create media players that support the specified file format.
+	 *
+	 * @param Url The URL to the media file container.
+	 * @see GetSupportedFormats
+	 */
+	bool SupportsFile( const FString& Url ) const
+	{
+		return GetSupportedFormats().Contains(FPaths::GetExtension(Url));
+	}
+
+public:
+
+	/**
+	 * Virtual destructor.
+	 */
+	virtual ~IMediaPlayerFactory( ) { }
+};
diff --git a/Engine/Source/Runtime/Media/Public/Interfaces/IMediaSink.h b/Engine/Source/Runtime/Media/Public/Interfaces/IMediaSink.h
new file mode 100644
--- /dev/null
+++ b/Engine/Source/Runtime/Media/Public/Interfaces/IMediaSink.h
@@ -0,0 +1,41 @@
+// Copyright 1998-2014 Epic Games, Inc. All Rights Reserved.
+
+#pragma once
+
+
+/** Type definition for shared pointers to instances of IMediaSink. */
+typedef TSharedPtr<class IMediaSink, ESPMode::ThreadSafe> IMediaSinkPtr;
+
+/** Type definition for shared references to instances of IMediaSink. */
+typedef TSharedRef<class IMediaSink, ESPMode::ThreadSafe> IMediaSinkRef;
+
+/** Type definition for weak pointers to instances of IMediaSink. */
+typedef TWeakPtr<class IMediaSink, ESPMode::ThreadSafe> IMediaSinkWeakPtr;
+
+
+/**
+ * Interface for recipients of media samples.
+ */
+class IMediaSink
+{
+public:
+
+	/**
+	 * Receives the given media sample data.
+	 *
+	 * Implementers of this method must be thread-safe, because it may get called from any thread.
+	 *
+	 * @param Buffer The buffer holding the sample data.
+	 * @param BufferSize The size of the buffer (in bytes).
+	 * @param Duration The duration of the sample data (MaxValue means no duration).
+	 * @param Time The time of the sample relative to the beginning of the media.
+	 */
+	virtual void ProcessMediaSample( const void* Buffer, uint32 BufferSize, FTimespan Duration, FTimespan Time ) = 0;
+
+public:
+
+	/**
+	 * Virtual destructor.
+	 */
+	~IMediaSink( ) { }
+};
diff --git a/Engine/Source/Runtime/Media/Public/Interfaces/IMediaTrack.h b/Engine/Source/Runtime/Media/Public/Interfaces/IMediaTrack.h
new file mode 100644
--- /dev/null
+++ b/Engine/Source/Runtime/Media/Public/Interfaces/IMediaTrack.h
@@ -0,0 +1,179 @@
+// Copyright 1998-2014 Epic Games, Inc. All Rights Reserved.
+
+#pragma once
+
+
+/**
+ * Enumerates supported media track types.
+ *
+ * The Media framework currently supports audio, caption and video tracks.
+ * We may implement additional types, such as HTML, still image and custom
+ * binary data tracks in the future.
+ */
+enum class EMediaTrackTypes
+{
+	/** Audio track. */
+	Audio,
+
+	/** Caption track (for subtitles). */
+	Caption,
+
+	/** Video track. */
+	Video
+};
+
+
+/** Type definition for shared pointers to instances of IMediaTrack. */
+typedef TSharedPtr<class IMediaTrack, ESPMode::ThreadSafe> IMediaTrackPtr;
+
+/** Type definition for shared references to instances of IMediaTrack. */
+typedef TSharedRef<class IMediaTrack, ESPMode::ThreadSafe> IMediaTrackRef;
+
+
+/**
+ * Interface for media tracks.
+ *
+ * A media track is a sequence of data of a uniform media type, such as audio or video.
+ * Media file containers may consist of a single track or they may contain multiple
+ * tracks of different types, for example video tracks with different encodings or
+ * multiple audio and caption text tracks for different languages.
+ *
+ * Track processing can be enabled and disabled while media playback is in progress.
+ * Use the Enable() and Disable() methods respectively. Whether a track is enabled by
+ * default is determined by the media file. For video files the video track as well
+ * as the audio track for the default language are usually enabled by default.
+ *
+ * Some track types may be mutually exclusive, which means that the media player can
+ * play only one track of a certain type at a time. For example, if a video contains
+ * multiple audio tracks, only one of them may be enabled. Use the IsMutuallyExclusive
+ * method to determine whether two tracks are mutually exclusive.
+ *
+ * @see IMediaPlayer
+ */
+class IMediaTrack
+{
+public:
+
+	/**
+	 * Adds the given media sink to this track.
+	 *
+	 * @param Sink The sink to add.
+	 * @see RemoveSink
+	 */
+	virtual void AddSink( const IMediaSinkRef& Sink ) = 0;
+
+	/**
+	 * Disables this track.
+	 *
+	 * If a track is disabled, no data will be decoded for it during media playback.
+	 *
+	 * @return true on success, false otherwise.
+	 * @see Enable, IsEnabled
+	 */
+	virtual bool Disable( ) = 0;
+
+	/**
+	 * Disables this track.
+	 *
+	 * If a track is enabled, data will be decoded for it during media playback.
+	 * Enabling a track may fail if the track is mutually exclusive with another track.
+	 * Use the IsMututallyExclusive method to check for mutual exclusion of tracks.
+	 *
+	 * @return true on success, false otherwise.
+	 * @see Disable, IsEnabled
+	 */
+	virtual bool Enable( ) = 0;
+
+	/**
+	 * Gets audio track related details.
+	 *
+	 * This method will assert if called on a track that it is not actually an audio track.
+	 * Use the GetType method to determine the type of track before calling this method.
+	 *
+	 * @return Audio track details.
+	 * @see GetCaptionDetails, GetVideoDetails, GetType
+	 */
+	virtual const IMediaTrackAudioDetails& GetAudioDetails( ) const = 0;
+
+	/**
+	 * Gets caption track related details.
+	 *
+	 * This method will assert if called on a track that it is not actually an caption track.
+	 * Use the GetType method to determine the type of track before calling this method.
+	 *
+	 * @return Caption track details.
+	 * @see GetAudioDetails, GetVideoDetails, GetType
+	 */
+	virtual const IMediaTrackCaptionDetails& GetCaptionDetails( ) const = 0;
+
+	/**
+	 * Gets the track's index number.
+	 *
+	 * @return Index number.
+	 */
+	virtual uint32 GetIndex( ) const = 0;
+
+	/**
+	 * Gets the track's language tag, i.e. "en-US" for English.
+	 *
+	 * @return Language tag.
+	 */
+	virtual FString GetLanguage( ) const = 0;
+
+	/**
+	 * Gets the name of the track.
+	 *
+	 * @return Track name string.
+	 */
+	virtual FString GetName( ) const = 0;
+
+	/**
+	 * Gets the track's type.
+	 *
+	 * @return Track type.
+	 */
+	virtual EMediaTrackTypes GetType( ) const = 0;
+
+	/**
+	 * Gets video track related details.
+	 *
+	 * This method will assert if called on a track that it is not actually an video track.
+	 * Use the GetType method to determine the type of track before calling this method.
+	 *
+	 * @return Video track details.
+	 * @see GetAudioDetails, GetCaptionDetails, GetType
+	 */
+	virtual const IMediaTrackVideoDetails& GetVideoDetails( ) const = 0;
+
+	/**
+	 * Checks whether this track is currently enabled.
+	 *
+	 * If a track is enabled, data will be decoded for it during media playback.
+	 *
+	 * @return true if the track is enabled, false otherwise.
+	 * @see Disable, Enable
+	 */
+	virtual bool IsEnabled( ) const = 0;
+
+	/**
+	 * Checks whether this track is mutually exclusive with another.
+	 *
+	 * @param Other The other track to check with.
+	 * @return true if the tracks are mutually exclusive, false otherwise.
+	 */
+	virtual bool IsMutuallyExclusive( const IMediaTrackRef& Other ) const = 0;
+
+	/**
+	 * Checks whether this track has protected content, i.e. DRM.
+	 *
+	 * @return true if the content is protected, false otherwise.
+	 */
+	virtual bool IsProtected( ) const = 0;
+
+	/**
+	 * Removes the given media sink from this track.
+	 *
+	 * @param Sink The sink to remove.
+	 */
+	virtual void RemoveSink( const IMediaSinkRef& Sink ) = 0;
+};
diff --git a/Engine/Source/Runtime/Media/Public/Interfaces/IMediaTrackAudioDetails.h b/Engine/Source/Runtime/Media/Public/Interfaces/IMediaTrackAudioDetails.h
new file mode 100644
--- /dev/null
+++ b/Engine/Source/Runtime/Media/Public/Interfaces/IMediaTrackAudioDetails.h
@@ -0,0 +1,33 @@
+// Copyright 1998-2014 Epic Games, Inc. All Rights Reserved.
+
+#pragma once
+
+
+/**
+ * Interface for audio track details.
+ */
+class IMediaTrackAudioDetails
+{
+public:
+
+	/**
+	 * Gets the number of audio channels.
+	 *
+	 * @return Number of channels.
+	 */
+	virtual uint32 GetNumChannels( ) const = 0;
+
+	/**
+	 * Gets the number of audio samples per second.
+	 *
+	 * @return Samples per second.
+	 */
+	virtual uint32 GetSamplesPerSecond( ) const = 0;
+
+public:
+
+	/**
+	 * Virtual destructor.
+	 */
+	virtual ~IMediaTrackAudioDetails( ) { }
+};
diff --git a/Engine/Source/Runtime/Media/Public/Interfaces/IMediaTrackCaptionDetails.h b/Engine/Source/Runtime/Media/Public/Interfaces/IMediaTrackCaptionDetails.h
new file mode 100644
--- /dev/null
+++ b/Engine/Source/Runtime/Media/Public/Interfaces/IMediaTrackCaptionDetails.h
@@ -0,0 +1,17 @@
+// Copyright 1998-2014 Epic Games, Inc. All Rights Reserved.
+
+#pragma once
+
+
+/**
+ * Interface for caption track details.
+ */
+class IMediaTrackCaptionDetails
+{
+public:
+
+	/**
+	 * Virtual destructor.
+	 */
+	virtual ~IMediaTrackCaptionDetails( ) { }
+};
diff --git a/Engine/Source/Runtime/Media/Public/Interfaces/IMediaTrackVideoDetails.h b/Engine/Source/Runtime/Media/Public/Interfaces/IMediaTrackVideoDetails.h
new file mode 100644
--- /dev/null
+++ b/Engine/Source/Runtime/Media/Public/Interfaces/IMediaTrackVideoDetails.h
@@ -0,0 +1,61 @@
+// Copyright 1998-2014 Epic Games, Inc. All Rights Reserved.
+
+#pragma once
+
+
+/**
+ * Interface for video tracks details.
+ */
+class IMediaTrackVideoDetails
+{
+public:
+
+	/**
+	 * Gets the average data rate of the video track in bits per second.
+	 *
+	 * @return The data rate.
+	 */
+	virtual uint32 GetBitRate( ) const = 0;
+
+	/**
+	 * Gets the video's dimensions in pixels.
+	 *
+	 * @return The width and height of the video.
+	 * @see GetAspectRatio
+	 */
+	virtual FIntPoint GetDimensions( ) const = 0;
+
+	/**
+	 * Gets the video's frame rate in frames per second.
+	 *
+	 * @return The frame rate.
+	 */
+	virtual float GetFrameRate( ) const = 0;
+
+public:
+
+	/**
+	 * Gets the media's aspect ratio.
+	 *
+	 * @return Aspect ratio.
+	 * @see GetDimensions
+	 */
+	float GetAspectRatio( ) const
+	{
+		const FVector2D Dimensions = GetDimensions();
+
+		if (FMath::IsNearlyZero(Dimensions.Y))
+		{
+			return 0.0f;
+		}
+
+		return Dimensions.X / Dimensions.Y;
+	}
+
+public:
+
+	/**
+	 * Virtual destructor.
+	 */
+	virtual ~IMediaTrackVideoDetails( ) { }
+};
diff --git a/Engine/Source/Runtime/Media/Public/Media.h b/Engine/Source/Runtime/Media/Public/Media.h
new file mode 100644
--- /dev/null
+++ b/Engine/Source/Runtime/Media/Public/Media.h
@@ -0,0 +1,22 @@
+// Copyright 1998-2014 Epic Games, Inc. All Rights Reserved.
+
+#pragma once
+
+
+/* Public Dependencies
+ *****************************************************************************/
+
+#include "ModuleManager.h"
+
+
+/* Public Includes
+ *****************************************************************************/
+
+#include "IMediaSink.h"
+#include "IMediaTrackAudioDetails.h"
+#include "IMediaTrackCaptionDetails.h"
+#include "IMediaTrackVideoDetails.h"
+#include "IMediaTrack.h"
+#include "IMediaPlayer.h"
+#include "IMediaPlayerFactory.h"
+#include "IMediaModule.h"
