Media: Mac CIS fix

diff --git a/Engine/Source/Runtime/Media/Public/Interfaces/IMediaModule.h b/Engine/Source/Runtime/Media/Public/Interfaces/IMediaModule.h
--- a/Engine/Source/Runtime/Media/Public/Interfaces/IMediaModule.h
+++ b/Engine/Source/Runtime/Media/Public/Interfaces/IMediaModule.h
@@ -73,7 +73,7 @@ public:
 	static IMediaModule& Get( )
 	{
 #if PLATFORM_IOS
-        static IMediaModule& Media = FModuleManager::LoadModuleChecked<IMediaModule>("Media");
+        static IMediaModule& MediaModule = FModuleManager::LoadModuleChecked<IMediaModule>("Media");
         return MediaModule;
 #else
         return FModuleManager::LoadModuleChecked<IMediaModule>("Media");
