Media: added missing code comments

diff --git a/Engine/Source/Runtime/Media/Public/Common/MediaSampleBuffer.h b/Engine/Source/Runtime/Media/Public/Common/MediaSampleBuffer.h
--- a/Engine/Source/Runtime/Media/Public/Common/MediaSampleBuffer.h
+++ b/Engine/Source/Runtime/Media/Public/Common/MediaSampleBuffer.h
@@ -24,8 +24,9 @@ public:
 		: CurrentSampleTime(FTimespan::MinValue())
 	{ }
 
-	virtual ~FMediaSampleBuffer()
-	{ }
+	/** Virtual destructor. */
+	virtual ~FMediaSampleBuffer( ) { }
+
 public:
 
 	/**
diff --git a/Engine/Source/Runtime/Media/Public/Common/MediaSampleQueue.h b/Engine/Source/Runtime/Media/Public/Common/MediaSampleQueue.h
--- a/Engine/Source/Runtime/Media/Public/Common/MediaSampleQueue.h
+++ b/Engine/Source/Runtime/Media/Public/Common/MediaSampleQueue.h
@@ -19,7 +19,10 @@ class FMediaSampleQueue
 {
 public:
 
-	virtual ~FMediaSampleQueue() {}
+	/** Virtual destructor. */
+	virtual ~FMediaSampleQueue( ) { }
+
+public:
 
 	// IMediaSink interface
 
