Media: Added API for retrieving media information

diff --git a/Engine/Plugins/Media/WmfMedia/Source/WmfMedia/Private/Player/WmfMediaPlayer.cpp b/Engine/Plugins/Media/WmfMedia/Source/WmfMedia/Private/Player/WmfMediaPlayer.cpp
--- a/Engine/Plugins/Media/WmfMedia/Source/WmfMedia/Private/Player/WmfMediaPlayer.cpp
+++ b/Engine/Plugins/Media/WmfMedia/Source/WmfMedia/Private/Player/WmfMediaPlayer.cpp
@@ -19,6 +19,54 @@ FWmfMediaPlayer::~FWmfMediaPlayer( )
 }
 
 
+/* IMediaInfo interface
+ *****************************************************************************/
+
+FTimespan FWmfMediaPlayer::GetDuration( ) const
+{
+	return Duration;
+}
+
+
+TRange<float> FWmfMediaPlayer::GetSupportedRates( EMediaPlaybackDirections Direction, bool Unthinned ) const
+{
+	if (MediaSession != NULL)
+	{
+		return MediaSession->GetSupportedRates(Direction, Unthinned);
+	}
+
+	return TRange<float>(0.0f);
+}
+
+
+FString FWmfMediaPlayer::GetUrl( ) const
+{
+	return MediaUrl;
+}
+
+
+bool FWmfMediaPlayer::SupportsRate( float Rate, bool Unthinned ) const
+{
+	return (MediaSession != NULL)
+		? MediaSession->IsRateSupported(Rate, Unthinned)
+		: false;
+}
+
+
+bool FWmfMediaPlayer::SupportsScrubbing( ) const
+{
+	return ((MediaSession != NULL) && MediaSession->SupportsScrubbing());
+}
+
+
+bool FWmfMediaPlayer::SupportsSeeking( ) const
+{
+	return ((MediaSession != NULL) &&
+			((MediaSession->GetCapabilities() & MFSESSIONCAP_SEEK) != 0) &&
+			(Duration > FTimespan::Zero()));
+}
+
+
 /* IMediaPlayer interface
  *****************************************************************************/
 
@@ -42,9 +90,9 @@ void FWmfMediaPlayer::Close( )
 }
 
 
-FTimespan FWmfMediaPlayer::GetDuration( ) const
+const IMediaInfo& FWmfMediaPlayer::GetMediaInfo( ) const 
 {
-	return Duration;
+	return *this;
 }
 
 
@@ -70,12 +118,6 @@ const TArray<IMediaTrackRef>& FWmfMediaPlayer::GetTracks( ) const
 }
 
 
-FString FWmfMediaPlayer::GetUrl( ) const
-{
-	return MediaUrl;
-}
-
-
 bool FWmfMediaPlayer::IsLooping( ) const 
 {
 	return ((MediaSession != NULL) && MediaSession->IsLooping());
@@ -211,27 +253,6 @@ bool FWmfMediaPlayer::SetRate( float Rate )
 }
 
 
-bool FWmfMediaPlayer::SupportsRate( float Rate, bool Unthinned ) const
-{
-	return (MediaSession != NULL)
-		? MediaSession->IsRateSupported(Rate, Unthinned)
-		: false;
-}
-
-
-bool FWmfMediaPlayer::SupportsScrubbing( ) const
-{
-	return ((MediaSession != NULL) && MediaSession->SupportsScrubbing());
-}
-
-
-bool FWmfMediaPlayer::SupportsSeeking( ) const
-{
-	return ((MediaSession != NULL) &&
-			((MediaSession->GetCapabilities() & MFSESSIONCAP_SEEK) != 0) &&
-			(Duration > FTimespan::Zero()));
-}
-
 
 /* FWmfMediaPlayer implementation
  *****************************************************************************/
diff --git a/Engine/Plugins/Media/WmfMedia/Source/WmfMedia/Private/Player/WmfMediaPlayer.h b/Engine/Plugins/Media/WmfMedia/Source/WmfMedia/Private/Player/WmfMediaPlayer.h
--- a/Engine/Plugins/Media/WmfMedia/Source/WmfMedia/Private/Player/WmfMediaPlayer.h
+++ b/Engine/Plugins/Media/WmfMedia/Source/WmfMedia/Private/Player/WmfMediaPlayer.h
@@ -9,7 +9,8 @@
  * Implements a media player using the Windows Media Foundation framework.
  */
 class FWmfMediaPlayer
-	: public IMediaPlayer
+	: public IMediaInfo
+	, public IMediaPlayer
 {
 public:
 
@@ -23,16 +24,26 @@ public:
 	 */
 	~FWmfMediaPlayer( );
 
+public:
+
+	// IMediaInfo interface
+
+	virtual FTimespan GetDuration( ) const override;
+	virtual TRange<float> GetSupportedRates( EMediaPlaybackDirections Direction, bool Unthinned ) const override;
+	virtual FString GetUrl( ) const override;
+	virtual bool SupportsRate( float Rate, bool Unthinned ) const override;
+	virtual bool SupportsScrubbing( ) const override;
+	virtual bool SupportsSeeking( ) const override;
+
 public:
 
 	// IMediaPlayer interface
 
 	virtual void Close( ) override;
-	virtual FTimespan GetDuration( ) const override;
+	virtual const IMediaInfo& GetMediaInfo( ) const override;
 	virtual float GetRate( ) const override;
 	virtual FTimespan GetTime( ) const override;
 	virtual const TArray<IMediaTrackRef>& GetTracks( ) const override;
-	virtual FString GetUrl( ) const override;
 	virtual bool IsLooping( ) const override;
 	virtual bool IsPaused( ) const override;
 	virtual bool IsPlaying( ) const override;
@@ -42,9 +53,6 @@ public:
 	virtual bool Seek( const FTimespan& Time ) override;
 	virtual bool SetLooping( bool Looping ) override;
 	virtual bool SetRate( float Rate ) override;
-	virtual bool SupportsRate( float Rate, bool Unthinned ) const override;
-	virtual bool SupportsScrubbing( ) const override;
-	virtual bool SupportsSeeking( ) const override;
 
 	DECLARE_DERIVED_EVENT(FWmfMediaPlayer, IMediaPlayer::FOnMediaClosing, FOnMediaClosing);
 	virtual FOnMediaClosing& OnClosing( ) override
diff --git a/Engine/Plugins/Media/WmfMedia/Source/WmfMedia/Private/Wmf/WmfMediaByteStream.h b/Engine/Plugins/Media/WmfMedia/Source/WmfMedia/Private/Wmf/WmfMediaByteStream.h
--- a/Engine/Plugins/Media/WmfMedia/Source/WmfMedia/Private/Wmf/WmfMediaByteStream.h
+++ b/Engine/Plugins/Media/WmfMedia/Source/WmfMedia/Private/Wmf/WmfMediaByteStream.h
@@ -21,10 +21,8 @@ public:
 	 */
 	FWmfMediaByteStream( const TSharedRef<TArray<uint8>>& InBuffer );
 
-	/** Virtual destructor */
-	virtual ~FWmfMediaByteStream()
-	{
-	}
+	/** Virtual destructor. */
+	virtual ~FWmfMediaByteStream( ) { }
 
 public:
 
diff --git a/Engine/Plugins/Media/WmfMedia/Source/WmfMedia/Private/Wmf/WmfMediaReadState.h b/Engine/Plugins/Media/WmfMedia/Source/WmfMedia/Private/Wmf/WmfMediaReadState.h
--- a/Engine/Plugins/Media/WmfMedia/Source/WmfMedia/Private/Wmf/WmfMediaReadState.h
+++ b/Engine/Plugins/Media/WmfMedia/Source/WmfMedia/Private/Wmf/WmfMediaReadState.h
@@ -30,10 +30,8 @@ public:
 		, RefCount(1)
 	{ }
 
-	/** Virtual destructor */
-	virtual ~FWmfMediaReadState()
-	{
-	}
+	/** Virtual destructor. */
+	virtual ~FWmfMediaReadState( ) { }
 
 public:
 
diff --git a/Engine/Plugins/Media/WmfMedia/Source/WmfMedia/Private/Wmf/WmfMediaSampler.h b/Engine/Plugins/Media/WmfMedia/Source/WmfMedia/Private/Wmf/WmfMediaSampler.h
--- a/Engine/Plugins/Media/WmfMedia/Source/WmfMedia/Private/Wmf/WmfMediaSampler.h
+++ b/Engine/Plugins/Media/WmfMedia/Source/WmfMedia/Private/Wmf/WmfMediaSampler.h
@@ -15,17 +15,13 @@ class FWmfMediaSampler
 {
 public:
 
-	/**
-	 * Default constructor.
-	 */
+	/** Default constructor. */
 	FWmfMediaSampler( )
 		: RefCount(1)
 	{ }
 
-	/** Virtual destructor */
-	virtual ~FWmfMediaSampler()
-	{
-	}
+	/** Virtual destructor. */
+	virtual ~FWmfMediaSampler( ) { }
 
 public:
 
diff --git a/Engine/Plugins/Media/WmfMedia/Source/WmfMedia/Private/Wmf/WmfMediaSession.cpp b/Engine/Plugins/Media/WmfMedia/Source/WmfMedia/Private/Wmf/WmfMediaSession.cpp
--- a/Engine/Plugins/Media/WmfMedia/Source/WmfMedia/Private/Wmf/WmfMediaSession.cpp
+++ b/Engine/Plugins/Media/WmfMedia/Source/WmfMedia/Private/Wmf/WmfMediaSession.cpp
@@ -4,7 +4,7 @@
 #include "AllowWindowsPlatformTypes.h"
 
 
-/* FWmfAsyncCallback structors
+/* FWmfMediaSession structors
  *****************************************************************************/
 
 FWmfMediaSession::FWmfMediaSession( const TComPtr<IMFTopology>& InTopology )
@@ -56,7 +56,7 @@ FWmfMediaSession::FWmfMediaSession( const TComPtr<IMFTopology>& InTopology )
 }
 
 
-/* FWmfAsyncCallback interface
+/* FWmfMediaSession interface
  *****************************************************************************/
 
 FTimespan FWmfMediaSession::GetPosition( ) const
@@ -75,6 +75,33 @@ FTimespan FWmfMediaSession::GetPosition( ) const
 }
 
 
+TRange<float> FWmfMediaSession::GetSupportedRates( EMediaPlaybackDirections Direction, bool Unthinned ) const
+{
+	if (RateSupport == NULL)
+	{
+		return TRange<float>(0.0f);
+	}
+
+	float MaxRate = 0.0f;
+	float MinRate = 0.0f;
+
+	switch (Direction)
+	{
+	case EMediaPlaybackDirections::Forward:
+		RateSupport->GetSlowestRate(MFRATE_FORWARD, !Unthinned, &MinRate);
+		RateSupport->GetFastestRate(MFRATE_FORWARD, !Unthinned, &MaxRate);
+		break;
+
+	case EMediaPlaybackDirections::Reverse:
+		RateSupport->GetSlowestRate(MFRATE_REVERSE, !Unthinned, &MinRate);
+		RateSupport->GetFastestRate(MFRATE_REVERSE, !Unthinned, &MaxRate);
+		break;
+	}
+
+	return TRange<float>(MinRate, MaxRate);
+}
+
+
 bool FWmfMediaSession::IsRateSupported( float Rate, bool Unthinned ) const
 {
 	if (Rate == 1.0f)
diff --git a/Engine/Plugins/Media/WmfMedia/Source/WmfMedia/Private/Wmf/WmfMediaSession.h b/Engine/Plugins/Media/WmfMedia/Source/WmfMedia/Private/Wmf/WmfMediaSession.h
--- a/Engine/Plugins/Media/WmfMedia/Source/WmfMedia/Private/Wmf/WmfMediaSession.h
+++ b/Engine/Plugins/Media/WmfMedia/Source/WmfMedia/Private/Wmf/WmfMediaSession.h
@@ -51,10 +51,8 @@ public:
 	 */
 	FWmfMediaSession( const TComPtr<IMFTopology>& InTopology );
 
-	/** Virtual destructor */
-	virtual ~FWmfMediaSession()
-	{
-	}
+	/** Virtual destructor .*/
+	virtual ~FWmfMediaSession( ) { }
 
 public:
 
@@ -107,6 +105,15 @@ public:
 		return CurrentState;
 	}
 
+	/**
+	 * Gets the supported playback rate for the specified rate type.
+	 *
+	 * @param Direction The playback direction.
+	 * @param Unthinned Whether the rates are for unthinned playback.
+	 * @return The supported playback rate.
+	 */
+	TRange<float> GetSupportedRates( EMediaPlaybackDirections Direction, bool Unthinned ) const;
+
 	/**
 	 * Checks whether playback is currently looping.
 	 *
diff --git a/Engine/Source/Runtime/Media/Public/Interfaces/IMediaInfo.h b/Engine/Source/Runtime/Media/Public/Interfaces/IMediaInfo.h
new file mode 100644
--- /dev/null
+++ b/Engine/Source/Runtime/Media/Public/Interfaces/IMediaInfo.h
@@ -0,0 +1,74 @@
+
+/**
+ * Enumerates directions for playback rate information.
+ */
+enum class EMediaPlaybackDirections
+{
+	/** Forward playback rates. */
+	Forward,
+
+	/** Reverse playback rates. */
+	Reverse
+};
+
+
+/**
+ * Interface for information about media loaded into a media player.
+ */
+class IMediaInfo
+{
+public:
+
+	/**
+	 * Gets the media's duration.
+	 *
+	 * @return A time span representing the duration.
+	 * @see GetTime
+	 */
+	virtual FTimespan GetDuration( ) const = 0;
+
+	/**
+	 * Gets the range of supported playback rates in the specified playback direction.
+	 *
+	 * @param Direction The playback direction.
+	 * @param Unthinned Whether the rates are for unthinned playback.
+	 * @return The range of supported rates.
+	 * @see SupportsRate
+	 */
+	virtual TRange<float> GetSupportedRates( EMediaPlaybackDirections Direction, bool Unthinned ) const = 0;
+
+	/**
+	 * Gets the URL of the currently loaded media.
+	 *
+	 * @return Media URL.
+	 */
+	virtual FString GetUrl( ) const = 0;
+
+	/**
+	 * Checks whether the specified playback rate is supported.
+	 *
+	 * @param Rate The rate to check (can be negative for reverse play).
+	 * @param Unthinned Whether no frames should be dropped at the given rate.
+	 * @return true if the rate is supported, false otherwise.
+	 * @see GetSupportedRates, SupportsScrubbing, SupportsSeeking
+	 */
+	virtual bool SupportsRate( float Rate, bool Unthinned ) const = 0;
+
+	/**
+	 * Checks whether scrubbing is supported.
+	 *
+	 * Scrubbing is the ability to decode video frames while seeking in a media item at a playback rate of 0.0.
+	 *
+	 * @return true if scrubbing is supported, false otherwise.
+	 * @see SupportsRate, SupportsSeeking
+	 */
+	virtual bool SupportsScrubbing( ) const = 0;
+
+	/**
+	 * Checks whether the currently loaded media can jump to certain times.
+	 *
+	 * @return true if seeking is supported, false otherwise.
+	 * @see SupportsRate, SupportsScrubbing
+	 */
+	virtual bool SupportsSeeking( ) const = 0;
+};
diff --git a/Engine/Source/Runtime/Media/Public/Interfaces/IMediaPlayer.h b/Engine/Source/Runtime/Media/Public/Interfaces/IMediaPlayer.h
--- a/Engine/Source/Runtime/Media/Public/Interfaces/IMediaPlayer.h
+++ b/Engine/Source/Runtime/Media/Public/Interfaces/IMediaPlayer.h
@@ -48,25 +48,24 @@ public:
 	virtual void Close( ) = 0;
 
 	/**
-	 * Gets the media's duration.
+	 * Gets the last error that occurred during media loading or playback.
 	 *
-	 * @return A time span representing the duration.
-	 * @see GetTime, Seek
+	 * @return The error string, or an empty string if no error occurred.
 	 */
-	virtual FTimespan GetDuration( ) const = 0;
+//	virtual FString GetLastError( ) const = 0;
 
 	/**
-	 * Gets the last error that occurred during media loading or playback.
+	 * Gets information about the currently loaded media.
 	 *
-	 * @return The error string, or an empty string if no error occurred.
+	 * @return Interface to media information.
 	 */
-//	virtual FString GetLastError( ) const = 0;
+	virtual const IMediaInfo& GetMediaInfo( ) const = 0;
 
 	/**
 	 * Gets the nominal playback rate, i.e. 1.0 for real time.
 	 *
 	 * @return Playback rate.
-	 * @see Pause, Play, SetRate, SupportsRate
+	 * @see Pause, Play, SetRate
 	 */
 	virtual float GetRate( ) const = 0;
 
@@ -74,7 +73,7 @@ public:
 	 * Gets the player's current playback time.
 	 *
 	 * @return Playback time.
-	 * @see GetDuration, Seek
+	 * @see Seek
 	 */
 	virtual FTimespan GetTime( ) const = 0;
 
@@ -85,13 +84,6 @@ public:
 	 */
 	virtual const TArray<IMediaTrackRef>& GetTracks( ) const = 0;
 
-	/**
-	 * Gets the URL of the currently loaded media.
-	 *
-	 * @return Media URL.
-	 */
-	virtual FString GetUrl( ) const = 0;
-
 	/**
 	 * Checks whether playback is currently looping.
 	 *
@@ -152,7 +144,7 @@ public:
 	 *
 	 * @param Time The playback time to set.
 	 * @return true on success, false otherwise.
-	 * @see GetDuration, GetTime
+	 * @see GetTime
 	 */
 	virtual bool Seek( const FTimespan& Time ) = 0;
 
@@ -174,38 +166,10 @@ public:
 	 *
 	 * @param Rate The playback rate to set.
 	 * @return true on success, false otherwise.
-	 * @see GetRate, Pause, Play, SupportsRate
+	 * @see GetRate, Pause, Play
 	 */
 	virtual bool SetRate( float Rate ) = 0;
 
-	/**
-	 * Checks whether the specified playback rate is supported.
-	 *
-	 * @param Rate The rate to check (can be negative for reverse play).
-	 * @param Unthinned Whether no frames should be dropped at the given rate.
-	 * @return true if the rate is supported, false otherwise.
-	 * @see GetRate, SetRate, SupportsScrubbing, SupportsSeeking
-	 */
-	virtual bool SupportsRate( float Rate, bool Unthinned ) const = 0;
-
-	/**
-	 * Checks whether scrubbing is supported.
-	 *
-	 * Scrubbing is the ability to decode video frames while seeking in a media item at a playback rate of 0.0.
-	 *
-	 * @return true if scrubbing is supported, false otherwise.
-	 * @see SupportsRate, SupportsSeeking
-	 */
-	virtual bool SupportsScrubbing( ) const = 0;
-
-	/**
-	 * Checks whether the currently loaded media can jump to certain times.
-	 *
-	 * @return true if seeking is supported, false otherwise.
-	 * @see SupportsRate, SupportsScrubbing
-	 */
-	virtual bool SupportsSeeking( ) const = 0;
-
 public:
 
 	/** Gets an event delegate that is invoked when media has been closed. */
@@ -280,10 +244,21 @@ public:
 
 		switch (Direction)
 		{
-		case EMediaSeekDirection::Backward: SeekTime = GetTime() - TimeOffset; break;
-		case EMediaSeekDirection::Beginning: SeekTime = TimeOffset; break;
-		case EMediaSeekDirection::End: SeekTime = GetDuration() - TimeOffset; break;
-		case EMediaSeekDirection::Forward: SeekTime = GetTime() + TimeOffset; break;
+		case EMediaSeekDirection::Backward:
+			SeekTime = GetTime() - TimeOffset;
+			break;
+
+		case EMediaSeekDirection::Beginning:
+			SeekTime = TimeOffset;
+			break;
+
+		case EMediaSeekDirection::End:
+			SeekTime = GetMediaInfo().GetDuration() - TimeOffset;
+			break;
+
+		case EMediaSeekDirection::Forward:
+			SeekTime = GetTime() + TimeOffset;
+			break;
 		}
 
 		return Seek(SeekTime);
diff --git a/Engine/Source/Runtime/Media/Public/Media.h b/Engine/Source/Runtime/Media/Public/Media.h
--- a/Engine/Source/Runtime/Media/Public/Media.h
+++ b/Engine/Source/Runtime/Media/Public/Media.h
@@ -17,6 +17,7 @@
 #include "IMediaTrackCaptionDetails.h"
 #include "IMediaTrackVideoDetails.h"
 #include "IMediaTrack.h"
+#include "IMediaInfo.h"
 #include "IMediaPlayer.h"
 #include "IMediaPlayerFactory.h"
 #include "IMediaModule.h"
diff --git a/Engine/Source/Runtime/MediaAssets/Private/Assets/MediaAsset.cpp b/Engine/Source/Runtime/MediaAssets/Private/Assets/MediaAsset.cpp
--- a/Engine/Source/Runtime/MediaAssets/Private/Assets/MediaAsset.cpp
+++ b/Engine/Source/Runtime/MediaAssets/Private/Assets/MediaAsset.cpp
@@ -34,7 +34,7 @@ FTimespan UMediaAsset::GetDuration( ) const
 {
 	if (MediaPlayer.IsValid())
 	{
-		return MediaPlayer->GetDuration();
+		return MediaPlayer->GetMediaInfo().GetDuration();
 	}
 
 	return FTimespan::Zero();
@@ -140,19 +140,19 @@ bool UMediaAsset::Seek( const FTimespan& InTime )
 
 bool UMediaAsset::SupportsRate( float Rate, bool Unthinned ) const
 {
-	return MediaPlayer.IsValid() && MediaPlayer->SupportsRate(Rate, Unthinned);
+	return MediaPlayer.IsValid() && MediaPlayer->GetMediaInfo().SupportsRate(Rate, Unthinned);
 }
 
 
 bool UMediaAsset::SupportsScrubbing( ) const
 {
-	return MediaPlayer.IsValid() && MediaPlayer->SupportsScrubbing();
+	return MediaPlayer.IsValid() && MediaPlayer->GetMediaInfo().SupportsScrubbing();
 }
 
 
 bool UMediaAsset::SupportsSeeking( ) const
 {
-	return MediaPlayer.IsValid() && MediaPlayer->SupportsSeeking();
+	return MediaPlayer.IsValid() && MediaPlayer->GetMediaInfo().SupportsSeeking();
 }
 
 
diff --git a/Engine/Source/Runtime/MediaAssets/Public/MediaAsset.h b/Engine/Source/Runtime/MediaAssets/Public/MediaAsset.h
--- a/Engine/Source/Runtime/MediaAssets/Public/MediaAsset.h
+++ b/Engine/Source/Runtime/MediaAssets/Public/MediaAsset.h
@@ -285,11 +285,11 @@ protected:
 	uint32 Looping:1;
 
 	/** Select where to stream the media from, i.e. file or memory. */
-	UPROPERTY(EditAnywhere, BlueprintReadOnly, Category=MediaSource)
+	UPROPERTY(EditAnywhere, BlueprintReadOnly, Category=Source)
 	TEnumAsByte<enum EMediaAssetStreamModes> StreamMode;
 
 	/** The URL to the media file to be played. */
-	UPROPERTY(EditAnywhere, Category=MediaSource)
+	UPROPERTY(EditAnywhere, Category=Source)
 	FFilePath URL;
 
 private:
