Media: Media assets now pick a suitable default track if none is selected

diff --git a/Engine/Source/Editor/DetailCustomizations/Private/MediaTextureCustomization.cpp b/Engine/Source/Editor/DetailCustomizations/Private/MediaTextureCustomization.cpp
--- a/Engine/Source/Editor/DetailCustomizations/Private/MediaTextureCustomization.cpp
+++ b/Engine/Source/Editor/DetailCustomizations/Private/MediaTextureCustomization.cpp
@@ -140,7 +140,7 @@ FText FMediaTextureCustomization::HandleVideoTrackComboButtonText( ) const
 		return LOCTEXT("NoMediaLoaded", "No media loaded");
 	}
 
-	IMediaTrackPtr Track = Player->GetTrackSafe(VideoTrackIndex, EMediaTrackTypes::Video);
+	IMediaTrackPtr Track = Player->GetTrack(VideoTrackIndex, EMediaTrackTypes::Video);
 
 	// generate track name string
 	if (Track.IsValid())
diff --git a/Engine/Source/Runtime/Media/Public/Interfaces/IMediaPlayer.h b/Engine/Source/Runtime/Media/Public/Interfaces/IMediaPlayer.h
--- a/Engine/Source/Runtime/Media/Public/Interfaces/IMediaPlayer.h
+++ b/Engine/Source/Runtime/Media/Public/Interfaces/IMediaPlayer.h
@@ -81,6 +81,7 @@ public:
 	 * Gets access to the media's audio, video and other tracks.
 	 *
 	 * @return Media tracks interface.
+	 * @see GetFirstTrack, GetTrack
 	 */
 	virtual const TArray<IMediaTrackRef>& GetTracks( ) const = 0;
 
@@ -183,13 +184,34 @@ public:
 public:
 
 	/**
-	 * Attempts to get the media track with the specified index and type.
+	 * Gets the first media track matching the specified type.
+	 *
+	 * @param TrackType The expected type of the track, i.e. audio or video.
+	 * @return The first matching track, nullptr otherwise.
+	 * @see GetTrack, GetTracks
+	 */
+	IMediaTrackPtr GetFirstTrack( EMediaTrackTypes TrackType )
+	{
+		for (const IMediaTrackRef& Track : GetTracks())
+		{
+			if (Track->GetType() == TrackType)
+			{
+				return Track;
+			}
+		}
+
+		return nullptr;
+	}
+
+	/**
+	 * Gets the media track with the specified index and type.
 	 *
 	 * @param TrackIndex The index of the track to get.
 	 * @param TrackType The expected type of the track, i.e. audio or video.
 	 * @return The track if it exists, nullptr otherwise.
+	 * @see GetFirstTrack, GetTracks
 	 */
-	IMediaTrackPtr GetTrackSafe( int32 TrackIndex, EMediaTrackTypes TrackType )
+	IMediaTrackPtr GetTrack( int32 TrackIndex, EMediaTrackTypes TrackType )
 	{
 		if (GetTracks().IsValidIndex(TrackIndex))
 		{
diff --git a/Engine/Source/Runtime/MediaAssets/Private/Assets/MediaSoundWave.cpp b/Engine/Source/Runtime/MediaAssets/Private/Assets/MediaSoundWave.cpp
--- a/Engine/Source/Runtime/MediaAssets/Private/Assets/MediaSoundWave.cpp
+++ b/Engine/Source/Runtime/MediaAssets/Private/Assets/MediaSoundWave.cpp
@@ -160,7 +160,17 @@ void UMediaSoundWave::InitializeTrack( )
 
 		if (Player .IsValid())
 		{
-			AudioTrack = Player ->GetTrackSafe(AudioTrackIndex, EMediaTrackTypes::Audio);
+			AudioTrack = Player ->GetTrack(AudioTrackIndex, EMediaTrackTypes::Audio);
+
+			if (!AudioTrack.IsValid())
+			{
+				AudioTrack = Player->GetFirstTrack(EMediaTrackTypes::Audio);
+
+				if (AudioTrack.IsValid())
+				{
+					AudioTrackIndex = AudioTrack->GetIndex();
+				}
+			}
 		}
 	}
 
diff --git a/Engine/Source/Runtime/MediaAssets/Private/Assets/MediaTexture.cpp b/Engine/Source/Runtime/MediaAssets/Private/Assets/MediaTexture.cpp
--- a/Engine/Source/Runtime/MediaAssets/Private/Assets/MediaTexture.cpp
+++ b/Engine/Source/Runtime/MediaAssets/Private/Assets/MediaTexture.cpp
@@ -195,7 +195,17 @@ void UMediaTexture::InitializeTrack( )
 
 		if (Player.IsValid())
 		{
-			VideoTrack = Player->GetTrackSafe(VideoTrackIndex, EMediaTrackTypes::Video);
+			VideoTrack = Player->GetTrack(VideoTrackIndex, EMediaTrackTypes::Video);
+
+			if (!VideoTrack.IsValid())
+			{
+				VideoTrack = Player->GetFirstTrack(EMediaTrackTypes::Video);
+
+				if (VideoTrack.IsValid())
+				{
+					VideoTrackIndex = VideoTrack->GetIndex();
+				}
+			}
 		}
 	}
 
diff --git a/Engine/Source/Runtime/MediaAssets/Public/MediaTexture.h b/Engine/Source/Runtime/MediaAssets/Public/MediaTexture.h
--- a/Engine/Source/Runtime/MediaAssets/Public/MediaTexture.h
+++ b/Engine/Source/Runtime/MediaAssets/Public/MediaTexture.h
@@ -119,11 +119,11 @@ protected:
 
 protected:
 
-	/**
-	 * Initializes the video track.
-	 */
+	/** Initializes the video track. */
 	void InitializeTrack( );
 
+	void PickDefaultTrack( );
+
 private:
 
 	/** Callback for when the UMediaPlayer asset changed its media. */
