Media: Made public headers compilable individually

diff --git a/Engine/Source/Runtime/Media/Private/MediaModule.cpp b/Engine/Source/Runtime/Media/Private/MediaModule.cpp
--- a/Engine/Source/Runtime/Media/Private/MediaModule.cpp
+++ b/Engine/Source/Runtime/Media/Private/MediaModule.cpp
@@ -1,6 +1,9 @@
 // Copyright 1998-2014 Epic Games, Inc. All Rights Reserved.
 
 #include "MediaPrivatePCH.h"
+#include "IMediaModule.h"
+#include "IMediaPlayerFactory.h"
+#include "IMediaPlayer.h"
 #include "ModuleManager.h"
 
 
@@ -88,13 +91,13 @@ public:
 
 private:
 
-	// Holds an event delegate that is invoked after a video player factory has been added.
+	/** Holds an event delegate that is invoked after a video player factory has been added. */
 	FOnFactoryAdded FactoryAddedEvent;
 
-	// Holds an event delegate that is invoked after a video player factory has been removed.
+	/** Holds an event delegate that is invoked after a video player factory has been removed. */
 	FOnFactoryRemoved FactoryRemovedEvent;
 
-	// Holds the registered video player factories.
+	/** Holds the registered video player factories. */
 	TArray<IMediaPlayerFactory*> MediaPlayerFactories;
 };
 
diff --git a/Engine/Source/Runtime/Media/Private/MediaPrivatePCH.h b/Engine/Source/Runtime/Media/Private/MediaPrivatePCH.h
--- a/Engine/Source/Runtime/Media/Private/MediaPrivatePCH.h
+++ b/Engine/Source/Runtime/Media/Private/MediaPrivatePCH.h
@@ -2,8 +2,6 @@
 
 #pragma once
 
-#include "Media.h"
-
 
 /* Private dependencies
  *****************************************************************************/
diff --git a/Engine/Source/Runtime/Media/Public/Helpers/MediaSampleBuffer.h b/Engine/Source/Runtime/Media/Public/Helpers/MediaSampleBuffer.h
--- a/Engine/Source/Runtime/Media/Public/Helpers/MediaSampleBuffer.h
+++ b/Engine/Source/Runtime/Media/Public/Helpers/MediaSampleBuffer.h
@@ -2,12 +2,7 @@
 
 #pragma once
 
-
-/** Type definition for shared pointers to instances of FMediaSampleBuffer. */
-typedef TSharedPtr<class FMediaSampleBuffer, ESPMode::ThreadSafe> FMediaSampleBufferPtr;
-
-/** Type definition for shared references to instances of FMediaSampleBuffer. */
-typedef TSharedRef<class FMediaSampleBuffer, ESPMode::ThreadSafe> FMediaSampleBufferRef;
+#include "IMediaSink.h"
 
 
 /**
@@ -73,3 +68,10 @@ private:
 	/** The playback time of the currently available media sample. */
 	FTimespan CurrentSampleTime;
 };
+
+
+/** Type definition for shared pointers to instances of FMediaSampleBuffer. */
+typedef TSharedPtr<FMediaSampleBuffer, ESPMode::ThreadSafe> FMediaSampleBufferPtr;
+
+/** Type definition for shared references to instances of FMediaSampleBuffer. */
+typedef TSharedRef<FMediaSampleBuffer, ESPMode::ThreadSafe> FMediaSampleBufferRef;
diff --git a/Engine/Source/Runtime/Media/Public/Helpers/MediaSampleQueue.h b/Engine/Source/Runtime/Media/Public/Helpers/MediaSampleQueue.h
--- a/Engine/Source/Runtime/Media/Public/Helpers/MediaSampleQueue.h
+++ b/Engine/Source/Runtime/Media/Public/Helpers/MediaSampleQueue.h
@@ -2,12 +2,7 @@
 
 #pragma once
 
-
-/** Type definition for shared pointers to instances of FMediaSampleQueue. */
-typedef TSharedPtr<class FMediaSampleQueue, ESPMode::ThreadSafe> FMediaSampleQueuePtr;
-
-/** Type definition for shared references to instances of FMediaSampleQueue. */
-typedef TSharedRef<class FMediaSampleQueue, ESPMode::ThreadSafe> FMediaSampleQueueRef;
+#include "IMediaSink.h"
 
 
 /**
@@ -35,3 +30,10 @@ public:
 		Enqueue(MakeShareable(Sample));
 	}
 };
+
+
+/** Type definition for shared pointers to instances of FMediaSampleQueue. */
+typedef TSharedPtr<FMediaSampleQueue, ESPMode::ThreadSafe> FMediaSampleQueuePtr;
+
+/** Type definition for shared references to instances of FMediaSampleQueue. */
+typedef TSharedRef<FMediaSampleQueue, ESPMode::ThreadSafe> FMediaSampleQueueRef;
diff --git a/Engine/Source/Runtime/Media/Public/IMediaInfo.h b/Engine/Source/Runtime/Media/Public/IMediaInfo.h
--- a/Engine/Source/Runtime/Media/Public/IMediaInfo.h
+++ b/Engine/Source/Runtime/Media/Public/IMediaInfo.h
@@ -1,3 +1,7 @@
+// Copyright 1998-2014 Epic Games, Inc. All Rights Reserved.
+
+#pragma once
+
 
 /**
  * Enumerates directions for playback rate information.
diff --git a/Engine/Source/Runtime/Media/Public/IMediaModule.h b/Engine/Source/Runtime/Media/Public/IMediaModule.h
--- a/Engine/Source/Runtime/Media/Public/IMediaModule.h
+++ b/Engine/Source/Runtime/Media/Public/IMediaModule.h
@@ -2,6 +2,14 @@
 
 #pragma once
 
+#include "ModuleInterface.h"
+#include "ModuleManager.h"
+
+
+// forward declarations
+class IMediaPlayer;
+class IMediaPlayerFactory;
+
 
 /**
  * Interface for media modules.
@@ -23,7 +31,7 @@ public:
 	 *
 	 * @param OutFormats Will hold the supported formats.
 	 */
-	virtual int32 GetSupportedFormats( FMediaFormats& OutFormats ) = 0;
+	virtual int32 GetSupportedFormats( TMap<FString, FText>& OutFormats ) = 0;
 
 	/**
 	 * Registers a media player factory.
diff --git a/Engine/Source/Runtime/Media/Public/IMediaPlayer.h b/Engine/Source/Runtime/Media/Public/IMediaPlayer.h
--- a/Engine/Source/Runtime/Media/Public/IMediaPlayer.h
+++ b/Engine/Source/Runtime/Media/Public/IMediaPlayer.h
@@ -2,6 +2,13 @@
 
 #pragma once
 
+#include "IMediaInfo.h"
+#include "IMediaTrack.h"
+
+
+// forward declaration
+class IMediaTrack;
+
 
 /**
  * Enumerates directions for seeking in media.
@@ -22,13 +29,6 @@ enum class EMediaSeekDirection
 };
 
 
-/** Type definition for shared pointers to instances of IMediaPlayer. */
-typedef TSharedPtr<class IMediaPlayer> IMediaPlayerPtr;
-
-/** Type definition for shared references to instances of IMediaPlayer. */
-typedef TSharedRef<class IMediaPlayer> IMediaPlayerRef;
-
-
 /**
  * Interface for media players.
  *
@@ -83,7 +83,7 @@ public:
 	 * @return Media tracks interface.
 	 * @see GetFirstTrack, GetTrack
 	 */
-	virtual const TArray<IMediaTrackRef>& GetTracks() const = 0;
+	virtual const TArray<TSharedRef<IMediaTrack, ESPMode::ThreadSafe>>& GetTracks() const = 0;
 
 	/**
 	 * Checks whether playback is currently looping.
@@ -190,9 +190,9 @@ public:
 	 * @return The first matching track, nullptr otherwise.
 	 * @see GetTrack, GetTracks
 	 */
-	IMediaTrackPtr GetFirstTrack( EMediaTrackTypes TrackType )
+	TSharedPtr<IMediaTrack, ESPMode::ThreadSafe> GetFirstTrack( EMediaTrackTypes TrackType )
 	{
-		for (const IMediaTrackRef& Track : GetTracks())
+		for (const TSharedRef<IMediaTrack, ESPMode::ThreadSafe>& Track : GetTracks())
 		{
 			if (Track->GetType() == TrackType)
 			{
@@ -211,11 +211,11 @@ public:
 	 * @return The track if it exists, nullptr otherwise.
 	 * @see GetFirstTrack, GetTracks
 	 */
-	IMediaTrackPtr GetTrack( int32 TrackIndex, EMediaTrackTypes TrackType )
+	TSharedPtr<IMediaTrack, ESPMode::ThreadSafe> GetTrack( int32 TrackIndex, EMediaTrackTypes TrackType )
 	{
 		if (GetTracks().IsValidIndex(TrackIndex))
 		{
-			IMediaTrackPtr Track = GetTracks()[TrackIndex];
+			TSharedPtr<IMediaTrack, ESPMode::ThreadSafe> Track = GetTracks()[TrackIndex];
 
 			if (Track->GetType() == TrackType)
 			{
@@ -291,3 +291,10 @@ public:
 	/** Virtual destructor. */
 	virtual ~IMediaPlayer() { }
 };
+
+
+/** Type definition for shared pointers to instances of IMediaPlayer. */
+typedef TSharedPtr<IMediaPlayer> IMediaPlayerPtr;
+
+/** Type definition for shared references to instances of IMediaPlayer. */
+typedef TSharedRef<IMediaPlayer> IMediaPlayerRef;
diff --git a/Engine/Source/Runtime/Media/Public/IMediaPlayerFactory.h b/Engine/Source/Runtime/Media/Public/IMediaPlayerFactory.h
--- a/Engine/Source/Runtime/Media/Public/IMediaPlayerFactory.h
+++ b/Engine/Source/Runtime/Media/Public/IMediaPlayerFactory.h
@@ -3,6 +3,10 @@
 #pragma once
 
 
+// forward declarations
+class IMediaPlayer;
+
+
 /**
  * Type definition for media file container formats.
  *
diff --git a/Engine/Source/Runtime/Media/Public/IMediaSink.h b/Engine/Source/Runtime/Media/Public/IMediaSink.h
--- a/Engine/Source/Runtime/Media/Public/IMediaSink.h
+++ b/Engine/Source/Runtime/Media/Public/IMediaSink.h
@@ -3,16 +3,6 @@
 #pragma once
 
 
-/** Type definition for shared pointers to instances of IMediaSink. */
-typedef TSharedPtr<class IMediaSink, ESPMode::ThreadSafe> IMediaSinkPtr;
-
-/** Type definition for shared references to instances of IMediaSink. */
-typedef TSharedRef<class IMediaSink, ESPMode::ThreadSafe> IMediaSinkRef;
-
-/** Type definition for weak pointers to instances of IMediaSink. */
-typedef TWeakPtr<class IMediaSink, ESPMode::ThreadSafe> IMediaSinkWeakPtr;
-
-
 /**
  * Interface for recipients of media samples.
  */
@@ -37,3 +27,13 @@ public:
 	/** Virtual destructor. */
 	~IMediaSink() { }
 };
+
+
+/** Type definition for shared pointers to instances of IMediaSink. */
+typedef TSharedPtr<IMediaSink, ESPMode::ThreadSafe> IMediaSinkPtr;
+
+/** Type definition for shared references to instances of IMediaSink. */
+typedef TSharedRef<IMediaSink, ESPMode::ThreadSafe> IMediaSinkRef;
+
+/** Type definition for weak pointers to instances of IMediaSink. */
+typedef TWeakPtr<IMediaSink, ESPMode::ThreadSafe> IMediaSinkWeakPtr;
diff --git a/Engine/Source/Runtime/Media/Public/IMediaTrack.h b/Engine/Source/Runtime/Media/Public/IMediaTrack.h
--- a/Engine/Source/Runtime/Media/Public/IMediaTrack.h
+++ b/Engine/Source/Runtime/Media/Public/IMediaTrack.h
@@ -3,6 +3,13 @@
 #pragma once
 
 
+// forward declarations
+class IMediaSink;
+class IMediaTrackAudioDetails;
+class IMediaTrackCaptionDetails;
+class IMediaTrackVideoDetails;
+
+
 /**
  * Enumerates supported media track types.
  *
@@ -23,13 +30,6 @@ enum class EMediaTrackTypes
 };
 
 
-/** Type definition for shared pointers to instances of IMediaTrack. */
-typedef TSharedPtr<class IMediaTrack, ESPMode::ThreadSafe> IMediaTrackPtr;
-
-/** Type definition for shared references to instances of IMediaTrack. */
-typedef TSharedRef<class IMediaTrack, ESPMode::ThreadSafe> IMediaTrackRef;
-
-
 /**
  * Interface for media tracks.
  *
@@ -60,7 +60,7 @@ public:
 	 * @param Sink The sink to add.
 	 * @see RemoveSink
 	 */
-	virtual void AddSink( const IMediaSinkRef& Sink ) = 0;
+	virtual void AddSink( const TSharedRef<IMediaSink, ESPMode::ThreadSafe>& Sink ) = 0;
 
 	/**
 	 * Disables this track.
@@ -173,7 +173,7 @@ public:
 	 * @param Other The other track to check with.
 	 * @return true if the tracks are mutually exclusive, false otherwise.
 	 */
-	virtual bool IsMutuallyExclusive( const IMediaTrackRef& Other ) const = 0;
+	virtual bool IsMutuallyExclusive( const TSharedRef<IMediaTrack, ESPMode::ThreadSafe>& Other ) const = 0;
 
 	/**
 	 * Checks whether this track has protected content, i.e. DRM.
@@ -187,10 +187,17 @@ public:
 	 *
 	 * @param Sink The sink to remove.
 	 */
-	virtual void RemoveSink( const IMediaSinkRef& Sink ) = 0;
+	virtual void RemoveSink( const TSharedRef<IMediaSink, ESPMode::ThreadSafe>& Sink ) = 0;
 
 public:
 
 	/** Virtual destructor. */
 	virtual ~IMediaTrack() { }
 };
+
+
+/** Type definition for shared pointers to instances of IMediaTrack. */
+typedef TSharedPtr<IMediaTrack, ESPMode::ThreadSafe> IMediaTrackPtr;
+
+/** Type definition for shared references to instances of IMediaTrack. */
+typedef TSharedRef<IMediaTrack, ESPMode::ThreadSafe> IMediaTrackRef;
