Media: Deprecating Media.h

#UpgradeNotes:
- Instead of including Media.h, please include the interface header files that you actually use, i.e. IMediaPlayer.h and/or IMediaTrack.h

diff --git a/Engine/Source/Runtime/Media/Public/Deprecated/Media.h b/Engine/Source/Runtime/Media/Public/Deprecated/Media.h
--- a/Engine/Source/Runtime/Media/Public/Deprecated/Media.h
+++ b/Engine/Source/Runtime/Media/Public/Deprecated/Media.h
@@ -2,6 +2,8 @@
 
 #pragma once
 
+EMIT_DEPRECATED_WARNING_MESSAGE("Including Media.h is deprecated. Please include the individual interface headers that you actually use!");
+
 
 /* Public Dependencies
  *****************************************************************************/
