Removed redundant comments.

diff --git a/Engine/Source/Runtime/Media/Public/IMediaPlayer.h b/Engine/Source/Runtime/Media/Public/IMediaPlayer.h
--- a/Engine/Source/Runtime/Media/Public/IMediaPlayer.h
+++ b/Engine/Source/Runtime/Media/Public/IMediaPlayer.h
@@ -6,7 +6,6 @@
 #include "IMediaTrack.h"
 
 
-// forward declaration
 class IMediaTrack;
 
 
diff --git a/Engine/Source/Runtime/Media/Public/IMediaPlayerFactory.h b/Engine/Source/Runtime/Media/Public/IMediaPlayerFactory.h
--- a/Engine/Source/Runtime/Media/Public/IMediaPlayerFactory.h
+++ b/Engine/Source/Runtime/Media/Public/IMediaPlayerFactory.h
@@ -3,7 +3,6 @@
 #pragma once
 
 
-// forward declarations
 class IMediaPlayer;
 
 
diff --git a/Engine/Source/Runtime/Media/Public/IMediaTrack.h b/Engine/Source/Runtime/Media/Public/IMediaTrack.h
--- a/Engine/Source/Runtime/Media/Public/IMediaTrack.h
+++ b/Engine/Source/Runtime/Media/Public/IMediaTrack.h
@@ -3,7 +3,6 @@
 #pragma once
 
 
-// forward declarations
 class IMediaSink;
 class IMediaTrackAudioDetails;
 class IMediaTrackCaptionDetails;
