Media: Fixed race condition crash in media sample buffer

diff --git a/Engine/Source/Runtime/Media/Public/Helpers/MediaSampleBuffer.h b/Engine/Source/Runtime/Media/Public/Helpers/MediaSampleBuffer.h
--- a/Engine/Source/Runtime/Media/Public/Helpers/MediaSampleBuffer.h
+++ b/Engine/Source/Runtime/Media/Public/Helpers/MediaSampleBuffer.h
@@ -32,6 +32,7 @@ public:
 	 */
 	TSharedPtr<TArray<uint8>, ESPMode::ThreadSafe> GetCurrentSample() const
 	{
+		FScopeLock Lock(&CriticalSection);
 		return CurrentSample;
 	}
 	
@@ -56,12 +57,16 @@ public:
 		Sample->AddUninitialized(BufferSize);
 		FMemory::Memcpy(Sample->GetData(), Buffer, BufferSize);
 
+		FScopeLock Lock(&CriticalSection);
 		CurrentSample = MakeShareable(Sample);
 		CurrentSampleTime = Time;
 	}
 	
 private:
 
+	/** Critical section for locking access to CurrentSample. */
+	mutable FCriticalSection CriticalSection;
+
 	/** Holds the sample data of the current media sample. */
 	TSharedPtr<TArray<uint8>, ESPMode::ThreadSafe> CurrentSample;
 
