MediaFramework: Fixed multi threaded access to sink in MediaClock.

diff --git a/Engine/Source/Runtime/Media/Private/MediaClock.cpp b/Engine/Source/Runtime/Media/Private/MediaClock.cpp
--- a/Engine/Source/Runtime/Media/Private/MediaClock.cpp
+++ b/Engine/Source/Runtime/Media/Private/MediaClock.cpp
@@ -21,6 +21,8 @@ FMediaClock::FMediaClock()
 
 void FMediaClock::TickFetch()
 {
+	UpdateSinkArray();
+
 	for (int32 SinkIndex = Sinks.Num() - 1; SinkIndex >= 0; --SinkIndex)
 	{
 		auto Sink = Sinks[SinkIndex].Pin();
@@ -39,6 +41,8 @@ void FMediaClock::TickFetch()
 
 void FMediaClock::TickInput()
 {
+	UpdateSinkArray();
+
 	for (int32 SinkIndex = Sinks.Num() - 1; SinkIndex >= 0; --SinkIndex)
 	{
 		auto Sink = Sinks[SinkIndex].Pin();
@@ -57,6 +61,8 @@ void FMediaClock::TickInput()
 
 void FMediaClock::TickOutput()
 {
+	UpdateSinkArray();
+
 	for (int32 SinkIndex = Sinks.Num() - 1; SinkIndex >= 0; --SinkIndex)
 	{
 		auto Sink = Sinks[SinkIndex].Pin();
@@ -75,6 +81,8 @@ void FMediaClock::TickOutput()
 
 void FMediaClock::TickRender()
 {
+	UpdateSinkArray();
+
 	for (int32 SinkIndex = Sinks.Num() - 1; SinkIndex >= 0; --SinkIndex)
 	{
 		auto Sink = Sinks[SinkIndex].Pin();
@@ -112,7 +120,15 @@ void FMediaClock::UpdateTimecode(const FTimespan NewTimecode, bool NewTimecodeLo
 
 void FMediaClock::AddSink(const TSharedRef<IMediaClockSink, ESPMode::ThreadSafe>& Sink)
 {
-	Sinks.AddUnique(Sink);
+	FScopeLock Lock(&SinkCriticalSection);
+
+	SinksToAdd.AddUnique(Sink);
+
+	// Just in case this sink is pending to be removed.
+	if (SinksToRemove.Contains(Sink))
+	{
+		SinksToRemove.RemoveSwap(Sink);
+	}
 }
 
 
@@ -130,5 +146,32 @@ bool FMediaClock::IsTimecodeLocked() const
 
 void FMediaClock::RemoveSink(const TSharedRef<IMediaClockSink, ESPMode::ThreadSafe>& Sink)
 {
-	Sinks.Remove(Sink);
+	FScopeLock Lock(&SinkCriticalSection);
+
+	SinksToRemove.AddUnique(Sink);
+
+	// Just in case this sink is pending to be added.
+	if (SinksToAdd.Contains(Sink))
+	{
+		SinksToAdd.RemoveSwap(Sink);
+	}
+}
+
+void FMediaClock::UpdateSinkArray()
+{
+	FScopeLock Lock(&SinkCriticalSection);
+
+	// Add new sinks.
+	for (auto Sink : SinksToAdd)
+	{
+		Sinks.AddUnique(Sink);
+	}
+	SinksToAdd.Empty();
+	
+	// Remove sinks.
+	for (auto Sink : SinksToRemove)
+	{
+		Sinks.RemoveSwap(Sink);
+	}
+	SinksToRemove.Empty();
 }
diff --git a/Engine/Source/Runtime/Media/Private/MediaClock.h b/Engine/Source/Runtime/Media/Private/MediaClock.h
--- a/Engine/Source/Runtime/Media/Private/MediaClock.h
+++ b/Engine/Source/Runtime/Media/Private/MediaClock.h
@@ -80,9 +80,24 @@ private:
 	/** Registered clock sinks. */
 	TArray<TWeakPtr<IMediaClockSink, ESPMode::ThreadSafe>> Sinks;
 
+	/** Sinks that need to be added. */
+	TArray<TWeakPtr<IMediaClockSink, ESPMode::ThreadSafe>> SinksToAdd;
+
+	/** Sinks that need to be removed. */
+	TArray<TWeakPtr<IMediaClockSink, ESPMode::ThreadSafe>> SinksToRemove;
+
+	/** Critical section for synchronizing access to sink. */
+	mutable FCriticalSection SinkCriticalSection;
+
 	/** The current time code. */
 	FTimespan Timecode;
 
 	/** Whether the time code is locked to an external clock. */
 	bool TimecodeLocked;
+
+	/**
+	 * Updates the Sink array to reflect what has been added and removed since the last call to UpdateSinkArray.
+	 */
+	void UpdateSinkArray();
+	
 };
