MediaFramework: Compile fix.

diff --git a/Engine/Source/Runtime/Media/Private/MediaClock.cpp b/Engine/Source/Runtime/Media/Private/MediaClock.cpp
--- a/Engine/Source/Runtime/Media/Private/MediaClock.cpp
+++ b/Engine/Source/Runtime/Media/Private/MediaClock.cpp
@@ -5,7 +5,7 @@
 #include "CoreTypes.h"
 
 #include "IMediaClockSink.h"
-
+#include "Misc/ScopeLock.h"
 
 /* FMediaClock structors
  *****************************************************************************/
