Fix Image Write not supported message.

diff --git a/Engine/Source/Runtime/Media/Private/IMediaTextureSample.cpp b/Engine/Source/Runtime/Media/Private/IMediaTextureSample.cpp
new file mode 100644
--- /dev/null
+++ b/Engine/Source/Runtime/Media/Private/IMediaTextureSample.cpp
@@ -0,0 +1,26 @@
+// Copyright 1998-2018 Epic Games, Inc. All Rights Reserved.
+
+#include "IMediaTextureSample.h"
+
+namespace MediaTextureSampleFormat
+{
+	const TCHAR* EnumToString(const EMediaTextureSampleFormat InSampleFormat)
+	{
+		switch(InSampleFormat)
+		{
+			case EMediaTextureSampleFormat::CharAYUV: return TEXT("CharAYUV");
+			case EMediaTextureSampleFormat::CharBGRA: return TEXT("CharBGRA");
+			case EMediaTextureSampleFormat::CharBGR10A2: return TEXT("CharBGR10A2");
+			case EMediaTextureSampleFormat::CharBMP: return TEXT("CharBMP");
+			case EMediaTextureSampleFormat::CharNV12: return TEXT("CharNV12");
+			case EMediaTextureSampleFormat::CharNV21: return TEXT("CharNV21");
+			case EMediaTextureSampleFormat::CharUYVY: return TEXT("CharUYVY");
+			case EMediaTextureSampleFormat::CharYUY2: return TEXT("CharYUY2");
+			case EMediaTextureSampleFormat::CharYVYU: return TEXT("CharYVYU");
+			case EMediaTextureSampleFormat::FloatRGB: return TEXT("FloatRGB");
+			case EMediaTextureSampleFormat::FloatRGBA: return TEXT("FloatRGBA");
+			case EMediaTextureSampleFormat::Undefined: 
+			default: return TEXT("Undefined");
+		}
+	}
+}
\ No newline at end of file
diff --git a/Engine/Source/Runtime/Media/Public/IMediaTextureSample.h b/Engine/Source/Runtime/Media/Public/IMediaTextureSample.h
--- a/Engine/Source/Runtime/Media/Public/IMediaTextureSample.h
+++ b/Engine/Source/Runtime/Media/Public/IMediaTextureSample.h
@@ -64,6 +64,10 @@ enum class EMediaTextureSampleFormat
 	FloatRGBA
 };
 
+namespace MediaTextureSampleFormat
+{
+	 MEDIA_API const TCHAR* EnumToString(const EMediaTextureSampleFormat InSampleFormat);
+};
 
 /**
  * Interface for media texture samples.
diff --git a/Engine/Source/Runtime/MediaUtils/Public/MediaRecorder.cpp b/Engine/Source/Runtime/MediaUtils/Public/MediaRecorder.cpp
--- a/Engine/Source/Runtime/MediaUtils/Public/MediaRecorder.cpp
+++ b/Engine/Source/Runtime/MediaUtils/Public/MediaRecorder.cpp
@@ -224,7 +224,7 @@ void FMediaRecorder::TickRecording()
 			{
 				if (!bUnsupportedWarningShowed)
 				{
-					UE_LOG(LogMediaUtils, Warning, TEXT("Texture Sample Format '%d' is not supported by Media Recorder."), (int32)Sample->GetFormat());
+					UE_LOG(LogMediaUtils, Warning, TEXT("Texture Sample Format '%s' is not supported by Media Recorder."), MediaTextureSampleFormat::EnumToString(Sample->GetFormat()));
 					bUnsupportedWarningShowed = true;
 				}
 				continue;
