Media ticker thread cleanup
- Avoid repeated memory allocations each tick with TickablesCopy array
- Avoid creating extra copies of TWeakPtr
- Fixed race condition that could lead to the thread to stop ticking
- Fixed 5ms delay and extra tick during shutdown
- Wait for infinite time instead of Timespan value

diff --git a/Engine/Source/Runtime/Media/Private/MediaTicker.cpp b/Engine/Source/Runtime/Media/Private/MediaTicker.cpp
--- a/Engine/Source/Runtime/Media/Private/MediaTicker.cpp
+++ b/Engine/Source/Runtime/Media/Private/MediaTicker.cpp
@@ -39,10 +39,13 @@ uint32 FMediaTicker::Run()
 {
 	while (!Stopping)
 	{
-		if (WakeupEvent->Wait(FTimespan::MaxValue()))
+		if (WakeupEvent->Wait() && !Stopping)
 		{
 			TickTickables();
-			FPlatformProcess::Sleep(0.005f);
+			if (!Stopping)
+			{
+				FPlatformProcess::Sleep(0.005f);
+			}
 		}
 	}
 
@@ -86,13 +89,13 @@ void FMediaTicker::RemoveTickable(const TSharedRef<IMediaTickable, ESPMode::Thre
 
 void FMediaTicker::TickTickables()
 {
-	TArray<TWeakPtr<IMediaTickable, ESPMode::ThreadSafe>> TickablesCopy;
+	TickablesCopy.Reset();
 	{
 		FScopeLock Lock(&CriticalSection);
 
 		for (int32 TickableIndex = Tickables.Num() - 1; TickableIndex >= 0; --TickableIndex)
 		{
-			auto Tickable = Tickables[TickableIndex].Pin();
+			TSharedPtr<IMediaTickable, ESPMode::ThreadSafe> Tickable = Tickables[TickableIndex].Pin();
 
 			if (Tickable.IsValid())
 			{
@@ -103,22 +106,20 @@ void FMediaTicker::TickTickables()
 				Tickables.RemoveAtSwap(TickableIndex);
 			}
 		}
-	}
 
-	if (TickablesCopy.Num() > 0)
-	{
-		for (auto TickablePtr : TickablesCopy)
+		if (Tickables.Num() == 0)
 		{
-			auto Tickable = TickablePtr.Pin();
-
-			if (Tickable.IsValid())
-			{
-				Tickable->TickTickable();
-			}
+			WakeupEvent->Reset();
 		}
 	}
-	else
+
+	for (int32 i=0; i < TickablesCopy.Num(); ++i)
 	{
-		WakeupEvent->Reset();
+		TSharedPtr<IMediaTickable, ESPMode::ThreadSafe> Tickable = TickablesCopy[i].Pin();
+
+		if (Tickable.IsValid())
+		{
+			Tickable->TickTickable();
+		}
 	}
 }
diff --git a/Engine/Source/Runtime/Media/Private/MediaTicker.h b/Engine/Source/Runtime/Media/Private/MediaTicker.h
--- a/Engine/Source/Runtime/Media/Private/MediaTicker.h
+++ b/Engine/Source/Runtime/Media/Private/MediaTicker.h
@@ -61,6 +61,9 @@ private:
 	/** Collection of tickable objects. */
 	TArray<TWeakPtr<IMediaTickable, ESPMode::ThreadSafe>> Tickables;
 
+	/** Variable to avoid rellocating the Tickables array repeatedly. */
+	TArray<TWeakPtr<IMediaTickable, ESPMode::ThreadSafe>> TickablesCopy;
+
 	/** Holds an event signaling the thread to wake up. */
 	FEvent* WakeupEvent;
 };
\ No newline at end of file
