MediaTicker updated to not sleep for 5 ms on shutdown

diff --git a/Engine/Source/Runtime/Media/Private/MediaTicker.cpp b/Engine/Source/Runtime/Media/Private/MediaTicker.cpp
--- a/Engine/Source/Runtime/Media/Private/MediaTicker.cpp
+++ b/Engine/Source/Runtime/Media/Private/MediaTicker.cpp
@@ -39,10 +39,13 @@ uint32 FMediaTicker::Run()
 {
 	while (!Stopping)
 	{
-		if (WakeupEvent->Wait(FTimespan::MaxValue()))
+		if (WakeupEvent->Wait())
 		{
 			TickTickables();
-			FPlatformProcess::Sleep(0.005f);
+			if (!Stopping)
+			{
+				FPlatformProcess::Sleep(0.005f);
+			}
 		}
 	}
 
