- Giving each platform a Guid forevermore, in DataDrivenPlatformInfo
- MediaModule now uses the platform Guid from DDPI instead of having a per-platform MediaInfo module that only supplied a Guid. Deleted many platform extension modules

diff --git a/Engine/Config/Android/DataDrivenPlatformInfo.ini b/Engine/Config/Android/DataDrivenPlatformInfo.ini
--- a/Engine/Config/Android/DataDrivenPlatformInfo.ini
+++ b/Engine/Config/Android/DataDrivenPlatformInfo.ini
@@ -3,6 +3,7 @@ bIsConfidential=false
 AudioCompressionSettingsIniSectionName=/Script/AndroidRuntimeSettings.AndroidRuntimeSettings
 Freezing_b32Bit=ini:Engine:Android:/Script/AndroidRuntimeSettings.AndroidRuntimeSettings:bBuildForArmV7
 Freezing_bForce64BitMemoryImagePointers=true
+GlobalIdentifier=3619EA87DE704A48BB1551754423C26A
 
 NormalIconPath=Launcher/Android/Platform_Android_24x
 LargeIconPath=Launcher/Android/Platform_Android_128x
diff --git a/Engine/Config/HoloLens/DataDrivenPlatformInfo.ini b/Engine/Config/HoloLens/DataDrivenPlatformInfo.ini
--- a/Engine/Config/HoloLens/DataDrivenPlatformInfo.ini
+++ b/Engine/Config/HoloLens/DataDrivenPlatformInfo.ini
@@ -2,6 +2,7 @@
 bIsConfidential=false
 Freezing_MaxFieldAlignment64=8
 Freezing_bAlignBases=true
+GlobalIdentifier=0DF604E112E4445280BEE1C74EB934B1
 
 NormalIconPath=Launcher/Windows/Platform_Windows_24x
 LargeIconPath=Launcher/Windows/Platform_Windows_128x
diff --git a/Engine/Config/IOS/DataDrivenPlatformInfo.ini b/Engine/Config/IOS/DataDrivenPlatformInfo.ini
--- a/Engine/Config/IOS/DataDrivenPlatformInfo.ini
+++ b/Engine/Config/IOS/DataDrivenPlatformInfo.ini
@@ -1,6 +1,7 @@
 [DataDrivenPlatformInfo]
 bIsConfidential=false
 AudioCompressionSettingsIniSectionName=/Script/IOSRuntimeSettings.IOSRuntimeSettings
+GlobalIdentifier=988EBA73F971495BAFB09639F8C796BD
 
 NormalIconPath=Launcher/IOS/Platform_IOS_24x
 LargeIconPath=Launcher/IOS/Platform_IOS_128x
diff --git a/Engine/Config/Linux/DataDrivenPlatformInfo.ini b/Engine/Config/Linux/DataDrivenPlatformInfo.ini
--- a/Engine/Config/Linux/DataDrivenPlatformInfo.ini
+++ b/Engine/Config/Linux/DataDrivenPlatformInfo.ini
@@ -1,5 +1,6 @@
 [DataDrivenPlatformInfo]
 IniParent=Unix
+GlobalIdentifier=115DE4FE241B465B970A872F3167492A
 
 NormalIconPath=Launcher/Linux/Platform_Linux_24x
 LargeIconPath=Launcher/Linux/Platform_Linux_128x
diff --git a/Engine/Config/LinuxAArch64/DataDrivenPlatformInfo.ini b/Engine/Config/LinuxAArch64/DataDrivenPlatformInfo.ini
--- a/Engine/Config/LinuxAArch64/DataDrivenPlatformInfo.ini
+++ b/Engine/Config/LinuxAArch64/DataDrivenPlatformInfo.ini
@@ -1,6 +1,7 @@
 [DataDrivenPlatformInfo]
 IniParent=Unix
 bIsConfidential=false
+GlobalIdentifier=C0B45A339DE340C7BCE24C4715C3BABF
 
 NormalIconPath=Launcher/Linux/Platform_Linux_24x
 LargeIconPath=Launcher/Linux/Platform_Linux_128x
diff --git a/Engine/Config/Lumin/DataDrivenPlatformInfo.ini b/Engine/Config/Lumin/DataDrivenPlatformInfo.ini
--- a/Engine/Config/Lumin/DataDrivenPlatformInfo.ini
+++ b/Engine/Config/Lumin/DataDrivenPlatformInfo.ini
@@ -1,6 +1,7 @@
 [DataDrivenPlatformInfo]
 bIsConfidential=false
 IniParent=Android
+GlobalIdentifier=A20FCE56EA274AADB4EA8567AA6BBAB3
 
 NormalIconPath=Launcher/Lumin/Platform_Lumin_24x
 LargeIconPath=Launcher/Lumin/Platform_Lumin_128x
diff --git a/Engine/Config/Mac/DataDrivenPlatformInfo.ini b/Engine/Config/Mac/DataDrivenPlatformInfo.ini
--- a/Engine/Config/Mac/DataDrivenPlatformInfo.ini
+++ b/Engine/Config/Mac/DataDrivenPlatformInfo.ini
@@ -1,6 +1,7 @@
 [DataDrivenPlatformInfo]
 bIsConfidential=false
 AudioCompressionSettingsIniSectionName=/Script/MacTargetPlatform.MacTargetSettings
+GlobalIdentifier=003BE29617004F0C8E1F786081EFBB1F
 
 NormalIconPath=Launcher/Mac/Platform_Mac_24x
 LargeIconPath=Launcher/Mac/Platform_Mac_128x
diff --git a/Engine/Config/TVOS/DataDrivenPlatformInfo.ini b/Engine/Config/TVOS/DataDrivenPlatformInfo.ini
--- a/Engine/Config/TVOS/DataDrivenPlatformInfo.ini
+++ b/Engine/Config/TVOS/DataDrivenPlatformInfo.ini
@@ -2,6 +2,7 @@
 bIsConfidential=false
 IniParent=IOS
 AudioCompressionSettingsIniSectionName=/Script/IOSRuntimeSettings.IOSRuntimeSettings
+GlobalIdentifier=A478294FBD0D4EC08830B6D4D219C1A4
 
 NormalIconPath=Launcher/TVOS/Platform_TVOS_24x
 LargeIconPath=Launcher/TVOS/Platform_TVOS_128x
diff --git a/Engine/Config/Windows/DataDrivenPlatformInfo.ini b/Engine/Config/Windows/DataDrivenPlatformInfo.ini
--- a/Engine/Config/Windows/DataDrivenPlatformInfo.ini
+++ b/Engine/Config/Windows/DataDrivenPlatformInfo.ini
@@ -6,6 +6,7 @@ Freezing_bWithRayTracing=true
 Freezing_MaxFieldAlignment32=4
 Freezing_MaxFieldAlignment64=8
 Freezing_bAlignBases=true
+GlobalIdentifier=D1D5F296FF834A87B20FAAA9D6B8E9A6
 
 NormalIconPath=Launcher/Windows/Platform_Windows_24x
 LargeIconPath=Launcher/Windows/Platform_Windows_128x
diff --git a/Engine/Source/Runtime/Core/Private/Misc/DataDrivenPlatformInfoRegistry.cpp b/Engine/Source/Runtime/Core/Private/Misc/DataDrivenPlatformInfoRegistry.cpp
--- a/Engine/Source/Runtime/Core/Private/Misc/DataDrivenPlatformInfoRegistry.cpp
+++ b/Engine/Source/Runtime/Core/Private/Misc/DataDrivenPlatformInfoRegistry.cpp
@@ -256,13 +256,18 @@ static void LoadDDPIIniSettings(const FConfigFile& IniFile, FDataDrivenPlatformI
 	DDPIGetBool(IniFile, TEXT("bIsFakePlatform"), Info.bIsFakePlatform);
 	DDPIGetString(IniFile, TEXT("AudioCompressionSettingsIniSectionName"), Info.AudioCompressionSettingsIniSectionName);
 	DDPIGetStringArray(IniFile, TEXT("AdditionalRestrictedFolders"), Info.AdditionalRestrictedFolders);
-	
+
 	DDPIGetBool(IniFile, TEXT("Freezing_b32Bit"), Info.Freezing_b32Bit);
 	DDPIGetUInt(IniFile, Info.Freezing_b32Bit ? TEXT("Freezing_MaxFieldAlignment32") : TEXT("Freezing_MaxFieldAlignment64"), Info.Freezing_MaxFieldAlignment);
 	DDPIGetBool(IniFile, TEXT("Freezing_bForce64BitMemoryImagePointers"), Info.Freezing_bForce64BitMemoryImagePointers);
 	DDPIGetBool(IniFile, TEXT("Freezing_bAlignBases"), Info.Freezing_bAlignBases);
 	DDPIGetBool(IniFile, TEXT("Freezing_bWithRayTracing"), Info.Freezing_bWithRayTracing);
 
+	FString GuidString;
+	DDPIGetString(IniFile, TEXT("GlobalIdentifier"), GuidString);
+	Info.GlobalIdentifier = FGuid(GuidString);
+	checkf(Info.GlobalIdentifier != FGuid(), TEXT("Platform %s didn't have a valid GlobalIdentifier set in DataDrivenPlatformInfo.ini"), *PlatformName.ToString());
+
 	// NOTE: add more settings here!
 
 #if DDPI_HAS_EXTENDED_PLATFORMINFO_DATA
diff --git a/Engine/Source/Runtime/Core/Public/Misc/DataDrivenPlatformInfoRegistry.h b/Engine/Source/Runtime/Core/Public/Misc/DataDrivenPlatformInfoRegistry.h
--- a/Engine/Source/Runtime/Core/Public/Misc/DataDrivenPlatformInfoRegistry.h
+++ b/Engine/Source/Runtime/Core/Public/Misc/DataDrivenPlatformInfoRegistry.h
@@ -117,6 +117,9 @@ struct FDataDrivenPlatformInfo
 	// list of additonal restricted folders
 	TArray<FString> AdditionalRestrictedFolders;
 
+	// GUID to prepresent this platform forever
+	FGuid GlobalIdentifier;
+
 	// MemoryFreezing information, matches FPlatformTypeLayoutParameters - defaults are clang, noneditor
 	uint32 Freezing_MaxFieldAlignment = 0xffffffff;
 	bool Freezing_b32Bit = false;
diff --git a/Engine/Source/Runtime/Media/Media.Build.cs b/Engine/Source/Runtime/Media/Media.Build.cs
--- a/Engine/Source/Runtime/Media/Media.Build.cs
+++ b/Engine/Source/Runtime/Media/Media.Build.cs
@@ -12,18 +12,10 @@ namespace UnrealBuildTool.Rules
 					"Core",
 				});
 
-			// Make sure we always have the info module for the common platforms
-			DynamicallyLoadedModuleNames.Add(GetMediaInfoModuleName());
-
 			PrivateIncludePaths.AddRange(
 				new string[] {
 					"Runtime/Media/Private",
 				});
 		}
-
-		protected virtual string GetMediaInfoModuleName()
-        {
-			return "MediaInfo";
-		}
 	}
 }
diff --git a/Engine/Source/Runtime/Media/Private/MediaModule.cpp b/Engine/Source/Runtime/Media/Private/MediaModule.cpp
--- a/Engine/Source/Runtime/Media/Private/MediaModule.cpp
+++ b/Engine/Source/Runtime/Media/Private/MediaModule.cpp
@@ -8,6 +8,7 @@
 #include "Modules/ModuleManager.h"
 #include "Stats/Stats.h"
 #include "ProfilingDebugging/CsvProfiler.h"
+#include "Misc/DataDrivenPlatformInfoRegistry.h"
 
 #include "IMediaCaptureSupport.h"
 #include "IMediaPlayerFactory.h"
@@ -26,36 +27,18 @@ CSV_DEFINE_CATEGORY_MODULE(MEDIA_API, MediaStreaming, false);
 class FMediaModule
 	: public IMediaModule
 {
-	struct FPlatformInfo
-	{
-		FPlatformInfo()
-		{ }
-
-		FPlatformInfo(const FName& PlatformName, const FGuid& PlatformGuid, IMediaInfo* MediaInfo)
-			: Name(PlatformName), Guid(PlatformGuid), Info(MediaInfo)
-		{ }
-
-		FName Name;
-		FGuid Guid;
-		IMediaInfo* Info;
-	};
 
 public:
 
 	//~ IMediaModule interface
 
-	virtual void RegisterPlatform(const FName& PlatformName, const FGuid& PlatformGuid, IMediaInfo *MediaInfo) override
-	{
-		PlatformInfo.Add(FPlatformInfo(PlatformName, PlatformGuid, MediaInfo));
-	}
-
 	virtual FName GetPlatformName(const FGuid& PlatformGuid) const override
 	{
-		for (const FPlatformInfo& Info : PlatformInfo)
+		for (auto& Pair : FDataDrivenPlatformInfoRegistry::GetAllPlatformInfos())
 		{
-			if (Info.Guid == PlatformGuid)
+			if (Pair.Value.GlobalIdentifier == PlatformGuid)
 			{
-				return Info.Name;
+				return Pair.Key;
 			}
 		}
 		return FName();
@@ -63,14 +46,7 @@ public:
 
 	virtual FGuid GetPlatformGuid(const FName& PlatformName) const override
 	{
-		for (const FPlatformInfo& Info : PlatformInfo)
-		{
-			if (Info.Name == PlatformName)
-			{
-				return Info.Guid;
-			}
-		}
-		return FGuid();
+		return FDataDrivenPlatformInfoRegistry::GetPlatformInfo(PlatformName).GlobalIdentifier;
 	}
 
 	virtual const TArray<IMediaCaptureSupport*>& GetCaptureSupports() const override
@@ -226,17 +202,6 @@ public:
 		{
 			TickerThread = FRunnableThread::Create(&Ticker, TEXT("FMediaTicker"));
 		}
-
-		TArray<FName> Modules;
-		FModuleManager::Get().FindModules(TEXT("*MediaInfo"), Modules);
-
-		for (int32 Index = 0; Index < Modules.Num(); Index++)
-		{
-			if (IMediaInfo* MediaInfo = FModuleManager::LoadModulePtr<IMediaInfo>(Modules[Index]))
-			{
-				MediaInfo->Initialize(this);
-			}
-		}
 	}
 
 	virtual void ShutdownModule() override
@@ -294,9 +259,6 @@ private:
 
 	/** The media clock's time source. */
 	TSharedPtr<IMediaTimeSource, ESPMode::ThreadSafe> TimeSource;
-
-	/** List of supported platforms */
-	TArray<FPlatformInfo> PlatformInfo;
 };
 
 IMPLEMENT_MODULE(FMediaModule, Media);
diff --git a/Engine/Source/Runtime/Media/Public/IMediaModule.h b/Engine/Source/Runtime/Media/Public/IMediaModule.h
--- a/Engine/Source/Runtime/Media/Public/IMediaModule.h
+++ b/Engine/Source/Runtime/Media/Public/IMediaModule.h
@@ -38,11 +38,6 @@ public:
 
 	//~ Platform management
 
-	/*
-	 * Register a platform
-	*/
-	virtual void RegisterPlatform(const FName & PlatformName, const FGuid& PlatformGuid, IMediaInfo* MediaInfo) = 0;
-
 	/*
 	 * Get a nice platform name from a GUID
 	*/
diff --git a/Engine/Source/Runtime/MediaInfo/MediaInfo.Build.cs b/Engine/Source/Runtime/MediaInfo/MediaInfo.Build.cs
deleted file mode 100644
diff --git a/Engine/Source/Runtime/MediaInfo/Private/MediaInfoModule.cpp b/Engine/Source/Runtime/MediaInfo/Private/MediaInfoModule.cpp
deleted file mode 100644
diff --git a/Engine/Source/Runtime/MediaInfo/Public/IMediaInfoModule.h b/Engine/Source/Runtime/MediaInfo/Public/IMediaInfoModule.h
deleted file mode 100644
