Media: Expose visible sets of tiles to the player through media view, and tiling description to the media sample.

Minimal API changes needed for virtual texturing support.

diff --git a/Engine/Source/Runtime/Media/Public/IMediaTextureSample.h b/Engine/Source/Runtime/Media/Public/IMediaTextureSample.h
--- a/Engine/Source/Runtime/Media/Public/IMediaTextureSample.h
+++ b/Engine/Source/Runtime/Media/Public/IMediaTextureSample.h
@@ -116,6 +116,18 @@ enum class EMediaOrientation
 class IMediaTextureSample
 {
 public:
+	/** Description of tiles (only used by tiled image sequences currently).*/
+	struct TilingDescription
+	{
+		FIntPoint TileNum  = FIntPoint::ZeroValue;
+		FIntPoint TileSize = FIntPoint::ZeroValue;
+		int32 TileBorderSize = 0;
+
+		FORCEINLINE bool IsValid() const
+		{
+			return TileNum.X > 0 && TileNum.Y > 0 && TileSize.X > 0 && TileSize.Y > 0;
+		}
+	};
 
 	/**
 	 * Get the sample's frame buffer.
@@ -149,6 +161,17 @@ public:
 		return 1;
 	}
 
+	/**
+	 * Get tile information (number, size and border size) of the sample.
+	 *
+	 * @return TileInfo struct
+	 * @note Default implementation provided as most samples will not feature tiles
+	 */
+	virtual TilingDescription GetTilingDescription() const
+	{
+		return TilingDescription();
+	}
+
 	/**
 	 * Get the amount of time for which the sample is valid.
 	 *
diff --git a/Engine/Source/Runtime/Media/Public/IMediaView.h b/Engine/Source/Runtime/Media/Public/IMediaView.h
--- a/Engine/Source/Runtime/Media/Public/IMediaView.h
+++ b/Engine/Source/Runtime/Media/Public/IMediaView.h
@@ -3,6 +3,62 @@
 #pragma once
 
 #include "Math/Quat.h"
+#include "Containers/Set.h"
+#include "Containers/Map.h"
+
+/**
+ * Tile coordinate of a tiled media source/texture.
+ */
+struct MEDIA_API FMediaTileCoordinate
+{
+	uint16 X, Y;
+
+	FORCEINLINE FMediaTileCoordinate()
+	{
+	}
+
+	FORCEINLINE explicit FMediaTileCoordinate(uint16 InX, uint16 InY)
+		: X(InX)
+		, Y(InY)
+	{
+	}
+
+	FORCEINLINE explicit FMediaTileCoordinate(int32 InX, int32 InY)
+		: X(IntCastChecked<uint16>(InX))
+		, Y(IntCastChecked<uint16>(InY))
+	{
+	}
+
+	FORCEINLINE FMediaTileCoordinate(EForceInit)
+		: X(0u)
+		, Y(0u)
+	{
+	}
+
+	FORCEINLINE bool operator==(const FMediaTileCoordinate& Other) const
+	{
+		return X == Other.X && Y == Other.Y;
+	}
+
+
+	FORCEINLINE bool operator!=(const FMediaTileCoordinate& Other) const
+	{
+		return X != Other.X || Y != Other.Y;
+	}
+
+	FORCEINLINE static FMediaTileCoordinate Zero()
+	{
+		return FMediaTileCoordinate(EForceInit::ForceInit);
+	}
+};
+
+template <> struct TIsPODType<FMediaTileCoordinate> { enum { Value = true }; };
+
+
+FORCEINLINE uint32 GetTypeHash(const FMediaTileCoordinate& Coord)
+{
+	return FCrc::TypeCrc32(Coord);
+}
 
 
 /**
@@ -23,6 +79,10 @@
  *
  * The user settings can be used by media players to customize the generated
  * audio or video output for specific users.
+ * 
+ * Note: This view class has now been extended to also support tiled
+ * media sources such as tiled exr image sequences. Sets of visible tiles can
+ * be specified per mip level.
  *
  * @see IMediaCache, IMediaControls, IMediaPlayer, IMediaSamples, IMediaTracks
  */
@@ -162,6 +222,33 @@ public:
 		return false; // override in child classes if supported
 	}
 
+public:
+	//~ Visible tile interface
+
+	/**
+	 * Get the visible tiles per mip-level.
+	 *
+	 * @param OutTiles Map of visible tile sets per mip level index.
+	 * @return true on success, false if feature is not available.
+	 * @see SetVisibleTiles
+	 */
+	virtual bool GetVisibleTiles(TMap<int32, TSet<FMediaTileCoordinate>>& OutTiles) const
+	{
+		return false; // override in child classes if supported
+	};
+
+	/**
+	 * Set the visible tiles per mip-level.
+	 *
+	 * @param InTiles Map of visible tile sets per mip level index.
+	 * @return true on success, false if feature is not available.
+	 * @see GetVisibleTiles
+	 */
+	virtual bool SetVisibleTiles(TMap<int32, TSet<FMediaTileCoordinate>>&& InTiles)
+	{
+		return false; // override in child classes if supported
+	};
+
 public:
 
 	/** Virtual destructor. */
