Media: Remove media texture sample tiling description from inner-class context to later allow forward declarations.

diff --git a/Engine/Plugins/Media/ImgMedia/Source/ImgMedia/Private/Loader/ImgMediaLoader.cpp b/Engine/Plugins/Media/ImgMedia/Source/ImgMedia/Private/Loader/ImgMediaLoader.cpp
--- a/Engine/Plugins/Media/ImgMedia/Source/ImgMedia/Private/Loader/ImgMediaLoader.cpp
+++ b/Engine/Plugins/Media/ImgMedia/Source/ImgMedia/Private/Loader/ImgMediaLoader.cpp
@@ -51,9 +51,9 @@ namespace ImgMediaLoader
 	}
 	
 
-	FImgMediaTextureSample::TilingDescription GetTilingDescription(const FImgMediaLoader* Loader)
+	FMediaTextureTilingDescription GetTilingDescription(const FImgMediaLoader* Loader)
 	{
-		FImgMediaTextureSample::TilingDescription TilingDesc;
+		FMediaTextureTilingDescription TilingDesc;
 		
 		if (Loader->IsTiled())
 		{
diff --git a/Engine/Plugins/Media/ImgMedia/Source/ImgMedia/Private/Player/ImgMediaTextureSample.h b/Engine/Plugins/Media/ImgMedia/Source/ImgMedia/Private/Player/ImgMediaTextureSample.h
--- a/Engine/Plugins/Media/ImgMedia/Source/ImgMedia/Private/Player/ImgMediaTextureSample.h
+++ b/Engine/Plugins/Media/ImgMedia/Source/ImgMedia/Private/Player/ImgMediaTextureSample.h
@@ -37,7 +37,7 @@ public:
 	 * @param InTime The sample time (in the player's local clock).
 	 * @param InDuration The duration for which the sample is valid.
 	 */
-	bool Initialize(FImgMediaFrame& InFrame, const FIntPoint& InOutputDim, FMediaTimeStamp InTime, FTimespan InDuration, uint8 InNumMipMaps, TilingDescription InTilingDesc = TilingDescription())
+	bool Initialize(FImgMediaFrame& InFrame, const FIntPoint& InOutputDim, FMediaTimeStamp InTime, FTimespan InDuration, uint8 InNumMipMaps, FMediaTextureTilingDescription InTilingDesc = FMediaTextureTilingDescription())
 	{
 		Duration = InDuration;
 		Frame = InFrame;
@@ -75,7 +75,7 @@ public:
 		return NumMipMaps;
 	}
 
-	virtual TilingDescription GetTilingDescription() const override
+	virtual FMediaTextureTilingDescription GetTilingDescription() const override
 	{
 		return TilingDesc;
 	}
@@ -145,5 +145,5 @@ private:
 	uint8 NumMipMaps;
 
 	/** Description of the number and size of tiles in this sample. */
-	TilingDescription TilingDesc;
+	FMediaTextureTilingDescription TilingDesc;
 };
diff --git a/Engine/Source/Runtime/Media/Public/IMediaTextureSample.h b/Engine/Source/Runtime/Media/Public/IMediaTextureSample.h
--- a/Engine/Source/Runtime/Media/Public/IMediaTextureSample.h
+++ b/Engine/Source/Runtime/Media/Public/IMediaTextureSample.h
@@ -91,6 +91,18 @@ namespace MediaTextureSampleFormat
 	 MEDIA_API const TCHAR* EnumToString(const EMediaTextureSampleFormat InSampleFormat);
 };
 
+/** Description of how the media texture sample is tiled (only used by tiled image sequences currently).*/
+struct FMediaTextureTilingDescription
+{
+	FIntPoint TileNum = FIntPoint::ZeroValue;
+	FIntPoint TileSize = FIntPoint::ZeroValue;
+	int32 TileBorderSize = 0;
+
+	FORCEINLINE bool IsValid() const
+	{
+		return TileNum.X > 0 && TileNum.Y > 0 && TileSize.X > 0 && TileSize.Y > 0;
+	}
+};
 
 enum class EMediaOrientation
 {
@@ -116,19 +128,6 @@ enum class EMediaOrientation
 class IMediaTextureSample
 {
 public:
-	/** Description of tiles (only used by tiled image sequences currently).*/
-	struct TilingDescription
-	{
-		FIntPoint TileNum  = FIntPoint::ZeroValue;
-		FIntPoint TileSize = FIntPoint::ZeroValue;
-		int32 TileBorderSize = 0;
-
-		FORCEINLINE bool IsValid() const
-		{
-			return TileNum.X > 0 && TileNum.Y > 0 && TileSize.X > 0 && TileSize.Y > 0;
-		}
-	};
-
 	/**
 	 * Get the sample's frame buffer.
 	 *
@@ -167,9 +166,9 @@ public:
 	 * @return TileInfo struct
 	 * @note Default implementation provided as most samples will not feature tiles
 	 */
-	virtual TilingDescription GetTilingDescription() const
+	virtual FMediaTextureTilingDescription GetTilingDescription() const
 	{
-		return TilingDescription();
+		return FMediaTextureTilingDescription();
 	}
 
 	/**
