Fixing up static analysis warning caused by recent submit
Horde Issue: 271085

diff --git a/Engine/Source/Runtime/Media/Public/IMediaTimeSource.h b/Engine/Source/Runtime/Media/Public/IMediaTimeSource.h
--- a/Engine/Source/Runtime/Media/Public/IMediaTimeSource.h
+++ b/Engine/Source/Runtime/Media/Public/IMediaTimeSource.h
@@ -47,25 +47,25 @@ public:
 		return (int64(PrimaryIndex) << 32) + int64(SecondaryIndex);
 	}
 
-	static int64 AdjustPrimaryIndex(int64 SequenceIndex, int32 Add)
+	static int64 AdjustPrimaryIndex(int64 InSequenceIndex, int32 Add)
 	{
-		return SequenceIndex + (int64(Add) << 32);
+		return InSequenceIndex + (int64(Add) << 32);
 	}
 
-	static int64 AdjustSecondaryIndex(int64 SequenceIndex, int32 Add)
+	static int64 AdjustSecondaryIndex(int64 InSequenceIndex, int32 Add)
 	{
-		return SequenceIndex + Add;
+		return InSequenceIndex + Add;
 	}
 
-	static int32 GetPrimaryIndex(int64 SequenceIndex)
+	static int32 GetPrimaryIndex(int64 InSequenceIndex)
 	{
 		// note: needs to cope with signed operation
-		return int32((SequenceIndex - int32(SequenceIndex)) >> 32);
+		return int32((InSequenceIndex - int32(InSequenceIndex)) >> 32);
 	}
 
-	static int32 GetSecondaryIndex(int64 SequenceIndex)
+	static int32 GetSecondaryIndex(int64 InSequenceIndex)
 	{
-		return int32(SequenceIndex);
+		return int32(InSequenceIndex);
 	}
 
 	FMediaTimeStamp operator + (const FTimespan & Other) const { return FMediaTimeStamp(Time + Other, SequenceIndex); }
