Added PVS static analysis warning suppression for false alram (Horde issue #271085)
#rnx

diff --git a/Engine/Source/Runtime/Media/Public/IMediaTimeSource.h b/Engine/Source/Runtime/Media/Public/IMediaTimeSource.h
--- a/Engine/Source/Runtime/Media/Public/IMediaTimeSource.h
+++ b/Engine/Source/Runtime/Media/Public/IMediaTimeSource.h
@@ -60,7 +60,7 @@ public:
 	static int32 GetPrimaryIndex(int64 InSequenceIndex)
 	{
 		// note: needs to cope with signed operation
-		return int32((InSequenceIndex - int32(InSequenceIndex)) >> 32);
+		return int32((InSequenceIndex - int32(InSequenceIndex)) >> 32);	//-V1065 PSV warning "Expression can be simplified" suppressed
 	}
 
 	static int32 GetSecondaryIndex(int64 InSequenceIndex)
