Media: Unused tile interface cleanup.

diff --git a/Engine/Source/Runtime/Media/Public/IMediaView.h b/Engine/Source/Runtime/Media/Public/IMediaView.h
--- a/Engine/Source/Runtime/Media/Public/IMediaView.h
+++ b/Engine/Source/Runtime/Media/Public/IMediaView.h
@@ -9,7 +9,7 @@
 /**
  * Tile coordinate of a tiled media source/texture.
  */
-struct FMediaTileCoordinate
+struct UE_DEPRECATED(5.7, "Unused since virtual texture support was never fully submitted.") FMediaTileCoordinate
 {
 	uint16 X, Y;
 
@@ -50,6 +50,7 @@ struct FMediaTileCoordinate
 	}
 };
 
+PRAGMA_DISABLE_DEPRECATION_WARNINGS
 template <> struct TIsPODType<FMediaTileCoordinate> { enum { Value = true }; };
 
 
@@ -57,6 +58,7 @@ inline uint32 GetTypeHash(const FMediaTileCoordinate& Coord)
 {
 	return FCrc::TypeCrc32(Coord);
 }
+PRAGMA_ENABLE_DEPRECATION_WARNINGS
 
 
 /**
@@ -221,6 +223,7 @@ public:
 	}
 
 public:
+PRAGMA_DISABLE_DEPRECATION_WARNINGS
 	//~ Visible tile interface
 
 	/**
@@ -230,6 +233,7 @@ public:
 	 * @return true on success, false if feature is not available.
 	 * @see SetVisibleTiles
 	 */
+	UE_DEPRECATED(5.7, "Unused since virtual texture support was never fully submitted.")
 	virtual bool GetVisibleTiles(TMap<int32, TSet<FMediaTileCoordinate>>& OutTiles) const
 	{
 		return false; // override in child classes if supported
@@ -242,10 +246,12 @@ public:
 	 * @return true on success, false if feature is not available.
 	 * @see GetVisibleTiles
 	 */
+	UE_DEPRECATED(5.7, "Unused since virtual texture support was never fully submitted.")
 	virtual bool SetVisibleTiles(TMap<int32, TSet<FMediaTileCoordinate>>&& InTiles)
 	{
 		return false; // override in child classes if supported
 	};
+PRAGMA_ENABLE_BUFFER_OVERRUN_WARNING
 
 public:
 
