Media: CIS warning quick fix due to pragma instruction mismatch in last CL.

diff --git a/Engine/Source/Runtime/Media/Public/IMediaView.h b/Engine/Source/Runtime/Media/Public/IMediaView.h
--- a/Engine/Source/Runtime/Media/Public/IMediaView.h
+++ b/Engine/Source/Runtime/Media/Public/IMediaView.h
@@ -251,7 +251,7 @@ PRAGMA_DISABLE_DEPRECATION_WARNINGS
 	{
 		return false; // override in child classes if supported
 	};
-PRAGMA_ENABLE_BUFFER_OVERRUN_WARNING
+PRAGMA_ENABLE_DEPRECATION_WARNINGS
 
 public:
 
