Media: Added new method GetTrackDuration() to IMediaTracks interface to query the duration of each track in the media source separately. This value is only returned by media players that support this. If the value is not provided an unset TOptional<> is returned.

#rnx

#Synced-CL 50901855

diff --git a/Engine/Plugins/Media/ElectraPlayer/Source/ElectraPlayerRuntime/Private/ElectraPlayer.cpp b/Engine/Plugins/Media/ElectraPlayer/Source/ElectraPlayerRuntime/Private/ElectraPlayer.cpp
--- a/Engine/Plugins/Media/ElectraPlayer/Source/ElectraPlayerRuntime/Private/ElectraPlayer.cpp
+++ b/Engine/Plugins/Media/ElectraPlayer/Source/ElectraPlayerRuntime/Private/ElectraPlayer.cpp
@@ -3584,6 +3584,13 @@ bool FElectraPlayer::GetVideoTrackFormat(int32 InTrackIndex, int32 InFormatIndex
 	return false;
 }
 
+TOptional<FTimespan> FElectraPlayer::GetTrackDuration(EMediaTrackType InTrackType, int32 InTrackIndex) const
+{
+	TSharedPtr<Electra::FTrackMetadata, ESPMode::ThreadSafe> Meta = GetTrackStreamMetadata(InTrackType, InTrackIndex);
+	return Meta.IsValid() ? Meta->Duration : TOptional<FTimespan>();
+}
+
+
 /**
  * Selects a specified track for playback.
  *
diff --git a/Engine/Plugins/Media/ElectraPlayer/Source/ElectraPlayerRuntime/Private/ElectraPlayer.h b/Engine/Plugins/Media/ElectraPlayer/Source/ElectraPlayerRuntime/Private/ElectraPlayer.h
--- a/Engine/Plugins/Media/ElectraPlayer/Source/ElectraPlayerRuntime/Private/ElectraPlayer.h
+++ b/Engine/Plugins/Media/ElectraPlayer/Source/ElectraPlayerRuntime/Private/ElectraPlayer.h
@@ -131,6 +131,7 @@ public:
 	bool SelectTrack(EMediaTrackType InTrackType, int32 InTrackIndex) override;
 	bool SetTrackFormat(EMediaTrackType InTrackType, int32 InTrackIndex, int32 InFormatIndex) override;
 	bool SetVideoTrackFrameRate(int32 InTrackIndex, int32 InFormatIndex, float InFrameRate) override;
+	TOptional<FTimespan> GetTrackDuration(EMediaTrackType InTrackType, int32 InTrackIndex) const override;
 
 	//-------------------------------------------------------------------------
 	// Methods from IAdaptiveStreamingPlayerMetrics
diff --git a/Engine/Plugins/Media/ElectraPlayer/Source/ElectraProtron/Private/Player/ElectraProtronPlayer.cpp b/Engine/Plugins/Media/ElectraPlayer/Source/ElectraProtron/Private/Player/ElectraProtronPlayer.cpp
--- a/Engine/Plugins/Media/ElectraPlayer/Source/ElectraProtron/Private/Player/ElectraProtronPlayer.cpp
+++ b/Engine/Plugins/Media/ElectraPlayer/Source/ElectraProtron/Private/Player/ElectraProtronPlayer.cpp
@@ -444,6 +444,12 @@ bool FElectraProtronPlayer::SetTrackFormat(EMediaTrackType InTrackType, int32 In
 	return Player.IsValid() ? Player->SetTrackFormat(InTrackType, InTrackIndex, InFormatIndex) : false;
 }
 
+TOptional<FTimespan> FElectraProtronPlayer::GetTrackDuration(EMediaTrackType InTrackType, int32 InTrackIndex) const
+{
+	TSharedPtr<FImpl, ESPMode::ThreadSafe> Player(GetCurrentPlayer());
+	return Player.IsValid() ? Player->GetTrackDuration(InTrackType, InTrackIndex) : TOptional<FTimespan>();
+}
+
 bool FElectraProtronPlayer::QueryCacheState(EMediaCacheState InState, TRangeSet<FTimespan>& OutTimeRanges) const
 {
 	TSharedPtr<FImpl, ESPMode::ThreadSafe> Player(GetCurrentPlayer());
diff --git a/Engine/Plugins/Media/ElectraPlayer/Source/ElectraProtron/Private/Player/ElectraProtronPlayer.h b/Engine/Plugins/Media/ElectraPlayer/Source/ElectraProtron/Private/Player/ElectraProtronPlayer.h
--- a/Engine/Plugins/Media/ElectraPlayer/Source/ElectraProtron/Private/Player/ElectraProtronPlayer.h
+++ b/Engine/Plugins/Media/ElectraPlayer/Source/ElectraProtron/Private/Player/ElectraProtronPlayer.h
@@ -97,6 +97,7 @@ protected:
 	bool GetVideoTrackFormat(int32 InTrackIndex, int32 InFormatIndex, FMediaVideoTrackFormat& OutFormat) const override;
 	bool SelectTrack(EMediaTrackType InTrackType, int32 InTrackIndex) override;
 	bool SetTrackFormat(EMediaTrackType InTrackType, int32 InTrackIndex, int32 InFormatIndex) override;
+	TOptional<FTimespan> GetTrackDuration(EMediaTrackType InTrackType, int32 InTrackIndex) const override;
 
 	//~ IMediaCache interface
 	bool QueryCacheState(EMediaCacheState InState, TRangeSet<FTimespan>& OutTimeRanges) const override;
diff --git a/Engine/Plugins/Media/ElectraPlayer/Source/ElectraProtron/Private/Player/ElectraProtronPlayerImpl.cpp b/Engine/Plugins/Media/ElectraPlayer/Source/ElectraProtron/Private/Player/ElectraProtronPlayerImpl.cpp
--- a/Engine/Plugins/Media/ElectraPlayer/Source/ElectraProtron/Private/Player/ElectraProtronPlayerImpl.cpp
+++ b/Engine/Plugins/Media/ElectraPlayer/Source/ElectraProtron/Private/Player/ElectraProtronPlayerImpl.cpp
@@ -1301,6 +1301,38 @@ bool FElectraProtronPlayer::FImpl::SetTrackFormat(EMediaTrackType InTrackType, i
 	return false;
 }
 
+TOptional<FTimespan> FElectraProtronPlayer::FImpl::GetTrackDuration(EMediaTrackType InTrackType, int32 InTrackIndex)
+{
+	if (InTrackIndex >= 0)
+	{
+		switch(InTrackType)
+		{
+			case EMediaTrackType::Video:
+			{
+				if (InTrackIndex < UsableTrackArrayIndicesByType[CodecTypeIndex(Electra::FCodecTypeFormat::EType::Video)].Num())
+				{
+					return Tracks[UsableTrackArrayIndicesByType[CodecTypeIndex(Electra::FCodecTypeFormat::EType::Video)][InTrackIndex]]->Duration;
+				}
+				break;
+			}
+			case EMediaTrackType::Audio:
+			{
+				if (InTrackIndex < UsableTrackArrayIndicesByType[CodecTypeIndex(Electra::FCodecTypeFormat::EType::Audio)].Num())
+				{
+					return Tracks[UsableTrackArrayIndicesByType[CodecTypeIndex(Electra::FCodecTypeFormat::EType::Audio)][InTrackIndex]]->Duration;
+				}
+				break;
+			}
+			default:
+			{
+				break;
+			}
+		}
+	}
+	return TOptional<FTimespan>();
+}
+
+
 bool FElectraProtronPlayer::FImpl::QueryCacheState(EMediaCacheState InState, TRangeSet<FTimespan>& OutTimeRanges)
 {
 	if (InState == EMediaCacheState::Loading)
diff --git a/Engine/Plugins/Media/ElectraPlayer/Source/ElectraProtron/Private/Player/ElectraProtronPlayerImpl.h b/Engine/Plugins/Media/ElectraPlayer/Source/ElectraProtron/Private/Player/ElectraProtronPlayerImpl.h
--- a/Engine/Plugins/Media/ElectraPlayer/Source/ElectraProtron/Private/Player/ElectraProtronPlayerImpl.h
+++ b/Engine/Plugins/Media/ElectraPlayer/Source/ElectraProtron/Private/Player/ElectraProtronPlayerImpl.h
@@ -250,6 +250,7 @@ public:
 	bool GetVideoTrackFormat(int32 InTrackIndex, int32 InFormatIndex, FMediaVideoTrackFormat& OutFormat);
 	bool SelectTrack(EMediaTrackType InTrackType, int32 InTrackIndex);
 	bool SetTrackFormat(EMediaTrackType InTrackType, int32 InTrackIndex, int32 InFormatIndex);
+	TOptional<FTimespan> GetTrackDuration(EMediaTrackType InTrackType, int32 InTrackIndex);
 	bool QueryCacheState(EMediaCacheState InState, TRangeSet<FTimespan>& OutTimeRanges);
 	int32 GetSampleCount(EMediaCacheState InState);
 
diff --git a/Engine/Source/Runtime/Media/Public/IMediaTracks.h b/Engine/Source/Runtime/Media/Public/IMediaTracks.h
--- a/Engine/Source/Runtime/Media/Public/IMediaTracks.h
+++ b/Engine/Source/Runtime/Media/Public/IMediaTracks.h
@@ -7,7 +7,8 @@
 #include "Math/IntPoint.h"
 #include "Math/Range.h"
 #include "Math/RangeSet.h"
-
+#include "Misc/Optional.h"
+#include "Misc/Timespan.h"
 
 /**
  * Enumerates available media track types.
@@ -88,6 +89,23 @@ class IMediaTracks
 {
 public:
 
+	/**
+	 * Returns the duration of a specific track.
+	 *
+	 * If a media file contains multiple tracks each track may have different duration and
+	 * the overall media duration is typically reported as that from the longest track.
+	 *
+	 * For some applications it may be required to know the exact track duration for processing,
+	 * which is what this method returns if the media container, as well as the media player
+	 * used for playback has the ability to provide it.
+	 *
+	 * @param TrackType The type of track to query.
+	 * @param TrackIndex The index of the track to query.
+	 * @return The duration of the track if known.
+	 */
+	virtual TOptional<FTimespan> GetTrackDuration(EMediaTrackType TrackType, int32 TrackIndex) const
+	{ return TOptional<FTimespan>(); }
+
 	/**
 	 * Get details about the specified audio track format.
 	 *
@@ -211,7 +229,7 @@ public:
 	 * Set the frame rate of the specified video track.
 	 *
 	 * Note that most players may not support overriding the video frame rate.
-	 * This feature is often only available on video capture media sources. 
+	 * This feature is often only available on video capture media sources.
 	 *
 	 * @param TrackIndex The index of the track (must be valid).
 	 * @param FormatIndex Index of the track format (must be valid).
diff --git a/Engine/Source/Runtime/MediaAssets/Private/Assets/MediaPlayer.cpp b/Engine/Source/Runtime/MediaAssets/Private/Assets/MediaPlayer.cpp
--- a/Engine/Source/Runtime/MediaAssets/Private/Assets/MediaPlayer.cpp
+++ b/Engine/Source/Runtime/MediaAssets/Private/Assets/MediaPlayer.cpp
@@ -426,6 +426,12 @@ FString UMediaPlayer::GetTrackLanguage(EMediaPlayerTrack TrackType, int32 TrackI
 }
 
 
+TOptional<FTimespan> UMediaPlayer::GetTrackDuration(EMediaPlayerTrack TrackType, int32 TrackIndex) const
+{
+	return PlayerFacade->GetTrackDuration((EMediaTrackType)TrackType, TrackIndex);
+}
+
+
 const FString& UMediaPlayer::GetUrl() const
 {
 	return PlayerFacade->GetUrl();
diff --git a/Engine/Source/Runtime/MediaAssets/Public/MediaPlayer.h b/Engine/Source/Runtime/MediaAssets/Public/MediaPlayer.h
--- a/Engine/Source/Runtime/MediaAssets/Public/MediaPlayer.h
+++ b/Engine/Source/Runtime/MediaAssets/Public/MediaPlayer.h
@@ -427,6 +427,16 @@ public:
 	 */
 	MEDIAASSETS_API TOptional<FTimecode> GetVideoTimecode() const;
 
+	/**
+	 * Gets the duration of a specific track, if available.
+	 * Availability depends on the media format and player implementation.
+	 * If unset only the overall media duration can be used.
+	 *
+	 * @return An optional time span representing the duration.
+	 * @see GetDuration
+	 */
+	MEDIAASSETS_API TOptional<FTimespan> GetTrackDuration(EMediaPlayerTrack TrackType, int32 TrackIndex) const;
+
 	/**
 	 * Get the human readable name of the specified track.
 	 *
diff --git a/Engine/Source/Runtime/MediaUtils/Private/MediaPlayerFacade.cpp b/Engine/Source/Runtime/MediaUtils/Private/MediaPlayerFacade.cpp
--- a/Engine/Source/Runtime/MediaUtils/Private/MediaPlayerFacade.cpp
+++ b/Engine/Source/Runtime/MediaUtils/Private/MediaPlayerFacade.cpp
@@ -613,6 +613,16 @@ FString FMediaPlayerFacade::GetTrackLanguage(EMediaTrackType TrackType, int32 Tr
 	return CurrentPlayer->GetTracks().GetTrackLanguage((EMediaTrackType)TrackType, TrackIndex);
 }
 
+TOptional<FTimespan> FMediaPlayerFacade::GetTrackDuration(EMediaTrackType TrackType, int32 TrackIndex) const
+{
+	TSharedPtr<IMediaPlayer, ESPMode::ThreadSafe> CurrentPlayer(Player);
+	if (!CurrentPlayer.IsValid())
+	{
+		return TOptional<FTimespan>();
+	}
+	return CurrentPlayer->GetTracks().GetTrackDuration((EMediaTrackType)TrackType, TrackIndex);
+}
+
 
 float FMediaPlayerFacade::GetVideoTrackAspectRatio(int32 TrackIndex, int32 FormatIndex) const
 {
diff --git a/Engine/Source/Runtime/MediaUtils/Public/MediaPlayerFacade.h b/Engine/Source/Runtime/MediaUtils/Public/MediaPlayerFacade.h
--- a/Engine/Source/Runtime/MediaUtils/Public/MediaPlayerFacade.h
+++ b/Engine/Source/Runtime/MediaUtils/Public/MediaPlayerFacade.h
@@ -441,6 +441,16 @@ public:
 	 */
 	MEDIAUTILS_API FString GetTrackLanguage(EMediaTrackType TrackType, int32 TrackIndex) const;
 
+	/**
+	 * Gets the duration of a specific track, if available.
+	 * Availability depends on the media format and player implementation.
+	 * If unset only the overall media duration can be used.
+	 *
+	 * @return An optional time span representing the duration.
+	 * @see GetDuration
+	 */
+	MEDIAUTILS_API TOptional<FTimespan> GetTrackDuration(EMediaTrackType TrackType, int32 TrackIndex) const;
+
 	/**
 	 * Get the URL of the currently loaded media, if any.
 	 *
